Remove GWorld (mostly) from Oculus code

diff --git a/Engine/Plugins/OculusRift/Source/OculusRift/Private/OculusRiftHMD.cpp b/Engine/Plugins/OculusRift/Source/OculusRift/Private/OculusRiftHMD.cpp
--- a/Engine/Plugins/OculusRift/Source/OculusRift/Private/OculusRiftHMD.cpp
+++ b/Engine/Plugins/OculusRift/Source/OculusRift/Private/OculusRiftHMD.cpp
@@ -368,7 +368,9 @@ void FOculusRiftHMD::ApplyHmdRotation(APlayerController* PC, FRotator& ViewRotat
 
 #if !UE_BUILD_SHIPPING
 	if (bDrawTrackingCameraFrustum)
-		DrawDebugTrackingCameraFrustum(PC->GetPawnOrSpectator()->GetPawnViewLocation());
+	{
+		DrawDebugTrackingCameraFrustum(PC->GetWorld(), PC->GetPawnOrSpectator()->GetPawnViewLocation());
+	}
 #endif
 }
 
@@ -406,12 +408,14 @@ void FOculusRiftHMD::UpdatePlayerCameraRotation(APlayerCameraManager* Camera, st
 
 #if !UE_BUILD_SHIPPING
 	if (bDrawTrackingCameraFrustum)
-		DrawDebugTrackingCameraFrustum(POV.Location);
+	{
+		DrawDebugTrackingCameraFrustum(Camera->GetWorld(), POV.Location);
+	}
 #endif
 }
 
 #if !UE_BUILD_SHIPPING
-void FOculusRiftHMD::DrawDebugTrackingCameraFrustum(const FVector& ViewLocation)
+void FOculusRiftHMD::DrawDebugTrackingCameraFrustum(UWorld* World, const FVector& ViewLocation)
 {
 	const FColor c = (HasValidTrackingPosition() ? FColor::Green : FColor::Red);
 	FVector origin;
@@ -420,7 +424,7 @@ void FOculusRiftHMD::DrawDebugTrackingCameraFrustum(const FVector& ViewLocation)
 	GetPositionalTrackingCameraProperties(origin, rotation, hfovDeg, vfovDeg, cameraDist, nearPlane, farPlane);
 
 	// Level line
-	//DrawDebugLine(GWorld, ViewLocation, FVector(ViewLocation.X + 1000, ViewLocation.Y, ViewLocation.Z), FColor::Blue);
+	//DrawDebugLine(World, ViewLocation, FVector(ViewLocation.X + 1000, ViewLocation.Y, ViewLocation.Z), FColor::Blue);
 
 	const float hfov = Math<float>::DegreeToRadFactor * hfovDeg * 0.5f;
 	const float vfov = Math<float>::DegreeToRadFactor * vfovDeg * 0.5f;
@@ -440,19 +444,19 @@ void FOculusRiftHMD::DrawDebugTrackingCameraFrustum(const FVector& ViewLocation)
 	coneBase4 = m.TransformPosition(coneBase4);
 
 	// draw a point at the camera pos
-	DrawDebugPoint(GWorld, coneTop, 5, c);
+	DrawDebugPoint(World, coneTop, 5, c);
 
 	// draw main pyramid, from top to base
-	DrawDebugLine(GWorld, coneTop, coneBase1, c);
-	DrawDebugLine(GWorld, coneTop, coneBase2, c);
-	DrawDebugLine(GWorld, coneTop, coneBase3, c);
-	DrawDebugLine(GWorld, coneTop, coneBase4, c);
+	DrawDebugLine(World, coneTop, coneBase1, c);
+	DrawDebugLine(World, coneTop, coneBase2, c);
+	DrawDebugLine(World, coneTop, coneBase3, c);
+	DrawDebugLine(World, coneTop, coneBase4, c);
 											  
 	// draw base (far plane)				  
-	DrawDebugLine(GWorld, coneBase1, coneBase2, c);
-	DrawDebugLine(GWorld, coneBase2, coneBase3, c);
-	DrawDebugLine(GWorld, coneBase3, coneBase4, c);
-	DrawDebugLine(GWorld, coneBase4, coneBase1, c);
+	DrawDebugLine(World, coneBase1, coneBase2, c);
+	DrawDebugLine(World, coneBase2, coneBase3, c);
+	DrawDebugLine(World, coneBase3, coneBase4, c);
+	DrawDebugLine(World, coneBase4, coneBase1, c);
 
 	// draw near plane
 	FVector coneNear1(-nearPlane, nearPlane * FMath::Tan(hfov), nearPlane * FMath::Tan(vfov));
@@ -463,16 +467,16 @@ void FOculusRiftHMD::DrawDebugTrackingCameraFrustum(const FVector& ViewLocation)
 	coneNear2 = m.TransformPosition(coneNear2);
 	coneNear3 = m.TransformPosition(coneNear3);
 	coneNear4 = m.TransformPosition(coneNear4);
-	DrawDebugLine(GWorld, coneNear1, coneNear2, c);
-	DrawDebugLine(GWorld, coneNear2, coneNear3, c);
-	DrawDebugLine(GWorld, coneNear3, coneNear4, c);
-	DrawDebugLine(GWorld, coneNear4, coneNear1, c);
+	DrawDebugLine(World, coneNear1, coneNear2, c);
+	DrawDebugLine(World, coneNear2, coneNear3, c);
+	DrawDebugLine(World, coneNear3, coneNear4, c);
+	DrawDebugLine(World, coneNear4, coneNear1, c);
 
 	// center line
 	FVector centerLine(-cameraDist, 0, 0);
 	centerLine = m.TransformPosition(centerLine);
-	DrawDebugLine(GWorld, coneTop, centerLine, FColor::Yellow);
-	DrawDebugPoint(GWorld, centerLine, 5, FColor::Yellow);
+	DrawDebugLine(World, coneTop, centerLine, FColor::Yellow);
+	DrawDebugPoint(World, centerLine, 5, FColor::Yellow);
 }
 #endif // #if !UE_BUILD_SHIPPING
 
diff --git a/Engine/Plugins/OculusRift/Source/OculusRift/Private/OculusRiftHMD.h b/Engine/Plugins/OculusRift/Source/OculusRift/Private/OculusRiftHMD.h
--- a/Engine/Plugins/OculusRift/Source/OculusRift/Private/OculusRiftHMD.h
+++ b/Engine/Plugins/OculusRift/Source/OculusRift/Private/OculusRiftHMD.h
@@ -210,7 +210,7 @@ private:
 	void		  SetHeadModel(const OVR::Vector3d&);
 
 #if !UE_BUILD_SHIPPING
-	void DrawDebugTrackingCameraFrustum(const FVector& ViewLocation);
+	void DrawDebugTrackingCameraFrustum(class UWorld* InWorld, const FVector& ViewLocation);
 #endif // #if !UE_BUILD_SHIPPING
 
 private: // data
