Fix Oculus Rift on Mac OS X again. Their SDK is erroneously reporting the Rift display origin as 0, 0 when it most definitely isn't so we fetch this information ourselves from CoreGraphics. We must also reposition the window on entering FullScreen on the Mac, so don't ignore that mode in SceneViewport.

diff --git a/Engine/Plugins/OculusRift/Source/OculusRift/Private/OculusRiftHMD.cpp b/Engine/Plugins/OculusRift/Source/OculusRift/Private/OculusRiftHMD.cpp
--- a/Engine/Plugins/OculusRift/Source/OculusRift/Private/OculusRiftHMD.cpp
+++ b/Engine/Plugins/OculusRift/Source/OculusRift/Private/OculusRiftHMD.cpp
@@ -823,6 +823,12 @@ bool FOculusRiftHMD::Exec( UWorld* InWorld, const TCHAR* Cmd, FOutputDevice& Ar
 		if (FParse::Command(&Cmd, TEXT("RESET")) && pHMD)
 		{
 			pHMD->GetDeviceInfo(&DeviceInfo);
+#if PLATFORM_MAC // @TODO Workaround broken Oculus SDK 25 - it reports 0,0 for display origin, even though it isn't.
+			CGDirectDisplayID DisplayId = (CGDirectDisplayID)DeviceInfo.DisplayId;
+			CGRect DisplayBounds = CGDisplayBounds(DisplayId);
+			DeviceInfo.DesktopX = DisplayBounds.origin.x;
+			DeviceInfo.DesktopY = DisplayBounds.origin.y;
+#endif
 			UpdateStereoRenderingParams();
 			return true;
 		}
@@ -1612,6 +1618,12 @@ void FOculusRiftHMD::Startup()
 
 		// Assuming we've successfully grabbed the device, read the configuration data from it, which we'll use for projection
 		pHMD->GetDeviceInfo(&DeviceInfo);
+#if PLATFORM_MAC // @TODO Workaround broken Oculus SDK 25 - it reports 0,0 for display origin, even though it isn't.
+		CGDirectDisplayID DisplayId = (CGDirectDisplayID)DeviceInfo.DisplayId;
+		CGRect DisplayBounds = CGDisplayBounds(DisplayId);
+		DeviceInfo.DesktopX = DisplayBounds.origin.x;
+		DeviceInfo.DesktopY = DisplayBounds.origin.y;
+#endif
 	}
 
 #if !UE_BUILD_SHIPPING
@@ -2149,6 +2161,12 @@ void  FOculusMessageHandler::Tick(float)
                 {
                     OVR::HMDInfo info;
                     desc.Handle.GetDeviceInfo(&info);
+#if PLATFORM_MAC // @TODO Workaround broken Oculus SDK 25 - it reports 0,0 for display origin, even though it isn't.
+					CGDirectDisplayID DisplayId = (CGDirectDisplayID)info.DisplayId;
+					CGRect DisplayBounds = CGDisplayBounds(DisplayId);
+					info.DesktopX = DisplayBounds.origin.x;
+					info.DesktopY = DisplayBounds.origin.y;
+#endif
                     // if strlen(info.DisplayDeviceName) == 0 then
                     // this HMD is 'fake' (created using sensor).
                     if (strlen(info.DisplayDeviceName) > 0 && (!pPlugin->pHMD || !info.IsSameDisplay(pPlugin->DeviceInfo)))
diff --git a/Engine/Source/Runtime/Engine/Private/Slate/SceneViewport.cpp b/Engine/Source/Runtime/Engine/Private/Slate/SceneViewport.cpp
--- a/Engine/Source/Runtime/Engine/Private/Slate/SceneViewport.cpp
+++ b/Engine/Source/Runtime/Engine/Private/Slate/SceneViewport.cpp
@@ -819,7 +819,8 @@ void FSceneViewport::ResizeFrame(uint32 NewSizeX, uint32 NewSizeY, EWindowMode::
 			{
 				if (GEngine->HMDDevice.IsValid() && GEngine->HMDDevice->IsHMDEnabled())
 				{
-					if (NewWindowMode && bNeedsResize)
+					// Resize & move only if moving to a fullscreen mode
+					if (NewWindowMode != EWindowMode::Windowed)
 					{
 						FSlateRect PreFullScreenRect = WindowToResize->GetRectInScreen();
 
