Fix for broken build when using clang compilers

diff --git a/Engine/Plugins/OculusRift/Source/OculusRift/Private/OculusRiftHMD.h b/Engine/Plugins/OculusRift/Source/OculusRift/Private/OculusRiftHMD.h
--- a/Engine/Plugins/OculusRift/Source/OculusRift/Private/OculusRiftHMD.h
+++ b/Engine/Plugins/OculusRift/Source/OculusRift/Private/OculusRiftHMD.h
@@ -18,7 +18,9 @@
 #define OVR_D3D_VERSION 11
 
 #ifdef OVR_VISION_ENABLED
-	#define OVR_CAPI_VISIONSUPPORT
+	#ifndef OVR_CAPI_VISIONSUPPORT
+		#define OVR_CAPI_VISIONSUPPORT
+	#endif
 #endif
 
 #if OCULUS_RIFT_SUPPORTED_PLATFORMS
