#oculus integrating fix for view matrix

diff --git a/Engine/Plugins/OculusRift/Source/OculusRift/Private/OculusRiftHMD.cpp b/Engine/Plugins/OculusRift/Source/OculusRift/Private/OculusRiftHMD.cpp
--- a/Engine/Plugins/OculusRift/Source/OculusRift/Private/OculusRiftHMD.cpp
+++ b/Engine/Plugins/OculusRift/Source/OculusRift/Private/OculusRiftHMD.cpp
@@ -1193,17 +1193,12 @@ FMatrix FOculusRiftHMD::GetStereoProjectionMatrix(enum EStereoscopicPass StereoP
 
 	FMatrix proj		 = ToFMatrix(EyeProjectionMatrices[idx]);
 
-	// nullify Z-component of translation vector, since otherwise it is causing
-	// rendering issues (far plane may stop rendering).
-	const FVector transV(EyeRenderDesc[idx].ViewAdjust.x, EyeRenderDesc[idx].ViewAdjust.y, 0.0f);
-	const FMatrix trans  = FMatrix::Identity.ConcatTranslation(transV);
-
 	// correct far and near planes
 	proj.M[2][2] = proj.M[3][3] = 0.0f;
 	proj.M[2][3] = 1.0f;
 	proj.M[3][2] = GNearClippingPlane;
 
-	return proj * trans;
+	return proj;
 }
 
 void FOculusRiftHMD::InitCanvasFromView(FSceneView* InView, UCanvas* Canvas)
