UE4 Refactoring. Changed OVERRIDE and FINAL macros to keywords override and final.

diff --git a/Engine/Plugins/Experimental/Paper2D/Source/Paper2DEditor/Private/SpriteEditor/SpriteEditorViewportClient.h b/Engine/Plugins/Experimental/Paper2D/Source/Paper2DEditor/Private/SpriteEditor/SpriteEditorViewportClient.h
--- a/Engine/Plugins/Experimental/Paper2D/Source/Paper2DEditor/Private/SpriteEditor/SpriteEditorViewportClient.h
+++ b/Engine/Plugins/Experimental/Paper2D/Source/Paper2DEditor/Private/SpriteEditor/SpriteEditorViewportClient.h
@@ -36,20 +36,20 @@ public:
 	FSpriteEditorViewportClient(TWeakPtr<FSpriteEditor> InSpriteEditor, TWeakPtr<class SSpriteEditorViewport> InSpriteEditorViewportPtr);
 
 	// FViewportClient interface
-	virtual void Draw(FViewport* Viewport, FCanvas* Canvas) OVERRIDE;
+	virtual void Draw(FViewport* Viewport, FCanvas* Canvas) override;
 	virtual void Draw(const FSceneView* View, FPrimitiveDrawInterface* PDI);
-	virtual void DrawCanvas(FViewport& InViewport, FSceneView& View, FCanvas& Canvas) OVERRIDE;
-	virtual void Tick(float DeltaSeconds) OVERRIDE;
+	virtual void DrawCanvas(FViewport& InViewport, FSceneView& View, FCanvas& Canvas) override;
+	virtual void Tick(float DeltaSeconds) override;
 	// End of FViewportClient interface
 
 	// FEditorViewportClient interface
-	virtual void UpdateMouseDelta() OVERRIDE;
-	virtual void ProcessClick(FSceneView& View, HHitProxy* HitProxy, FKey Key, EInputEvent Event, uint32 HitX, uint32 HitY) OVERRIDE;
-	virtual bool InputKey(FViewport* Viewport, int32 ControllerId, FKey Key, EInputEvent Event, float AmountDepressed, bool bGamepad) OVERRIDE;
-	virtual bool InputWidgetDelta(FViewport* Viewport, EAxisList::Type CurrentAxis, FVector& Drag, FRotator& Rot, FVector& Scale) OVERRIDE;
-	virtual void TrackingStarted(const struct FInputEventState& InInputState, bool bIsDragging, bool bNudge) OVERRIDE;
-	virtual void TrackingStopped() OVERRIDE;
-	virtual FWidget::EWidgetMode GetWidgetMode() const OVERRIDE;
+	virtual void UpdateMouseDelta() override;
+	virtual void ProcessClick(FSceneView& View, HHitProxy* HitProxy, FKey Key, EInputEvent Event, uint32 HitX, uint32 HitY) override;
+	virtual bool InputKey(FViewport* Viewport, int32 ControllerId, FKey Key, EInputEvent Event, float AmountDepressed, bool bGamepad) override;
+	virtual bool InputWidgetDelta(FViewport* Viewport, EAxisList::Type CurrentAxis, FVector& Drag, FRotator& Rot, FVector& Scale) override;
+	virtual void TrackingStarted(const struct FInputEventState& InInputState, bool bIsDragging, bool bNudge) override;
+	virtual void TrackingStopped() override;
+	virtual FWidget::EWidgetMode GetWidgetMode() const override;
 	virtual FVector GetWidgetLocation() const;
 	virtual FMatrix GetWidgetCoordSystem() const;
 	virtual ECoordSystem GetWidgetCoordSystemSpace() const;
diff --git a/Engine/Plugins/Experimental/Paper2D/Source/Paper2DEditor/Private/SpriteEditor/SpritePolygonCollectionCustomization.h b/Engine/Plugins/Experimental/Paper2D/Source/Paper2DEditor/Private/SpriteEditor/SpritePolygonCollectionCustomization.h
--- a/Engine/Plugins/Experimental/Paper2D/Source/Paper2DEditor/Private/SpriteEditor/SpritePolygonCollectionCustomization.h
+++ b/Engine/Plugins/Experimental/Paper2D/Source/Paper2DEditor/Private/SpriteEditor/SpritePolygonCollectionCustomization.h
@@ -14,8 +14,8 @@ public:
 	static TSharedRef<IPropertyTypeCustomization> MakeInstance();
 
 	// IPropertyTypeCustomization interface
-	virtual void CustomizeHeader(TSharedRef<IPropertyHandle> StructPropertyHandle, class FDetailWidgetRow& HeaderRow, IPropertyTypeCustomizationUtils& StructCustomizationUtils) OVERRIDE;
-	virtual void CustomizeChildren(TSharedRef<class IPropertyHandle> StructPropertyHandle, class IDetailChildrenBuilder& StructBuilder, IPropertyTypeCustomizationUtils& StructCustomizationUtils) OVERRIDE;
+	virtual void CustomizeHeader(TSharedRef<IPropertyHandle> StructPropertyHandle, class FDetailWidgetRow& HeaderRow, IPropertyTypeCustomizationUtils& StructCustomizationUtils) override;
+	virtual void CustomizeChildren(TSharedRef<class IPropertyHandle> StructPropertyHandle, class IDetailChildrenBuilder& StructBuilder, IPropertyTypeCustomizationUtils& StructCustomizationUtils) override;
 	// End of IPropertyTypeCustomization interface
 
 protected:
diff --git a/Engine/Plugins/Experimental/Paper2D/Source/Paper2DEditor/Private/TileMapEditing/EdModeTileMap.h b/Engine/Plugins/Experimental/Paper2D/Source/Paper2DEditor/Private/TileMapEditing/EdModeTileMap.h
--- a/Engine/Plugins/Experimental/Paper2D/Source/Paper2DEditor/Private/TileMapEditing/EdModeTileMap.h
+++ b/Engine/Plugins/Experimental/Paper2D/Source/Paper2DEditor/Private/TileMapEditing/EdModeTileMap.h
@@ -34,33 +34,33 @@ public:
 	virtual ~FEdModeTileMap();
 
 	// FSerializableObject interface
-	virtual void AddReferencedObjects(FReferenceCollector& Collector) OVERRIDE;
+	virtual void AddReferencedObjects(FReferenceCollector& Collector) override;
 	// End of FSerializableObject
 
 	// FEdMode interface
-	virtual bool UsesToolkits() const OVERRIDE;
-	virtual void Enter() OVERRIDE;
-	virtual void Exit() OVERRIDE;
+	virtual bool UsesToolkits() const override;
+	virtual void Enter() override;
+	virtual void Exit() override;
 	//virtual bool BoxSelect(FBox& InBox, bool bInSelect) OVERRIDE;
 
 	//	virtual void PostUndo() OVERRIDE;
-	virtual bool MouseMove(FLevelEditorViewportClient* ViewportClient, FViewport* Viewport, int32 x, int32 y) OVERRIDE;
-	virtual bool CapturedMouseMove(FLevelEditorViewportClient* InViewportClient, FViewport* InViewport, int32 InMouseX, int32 InMouseY) OVERRIDE;
-	virtual bool StartTracking(FLevelEditorViewportClient* InViewportClient, FViewport* InViewport) OVERRIDE;
-	virtual bool EndTracking(FLevelEditorViewportClient* InViewportClient, FViewport* InViewport) OVERRIDE;
+	virtual bool MouseMove(FLevelEditorViewportClient* ViewportClient, FViewport* Viewport, int32 x, int32 y) override;
+	virtual bool CapturedMouseMove(FLevelEditorViewportClient* InViewportClient, FViewport* InViewport, int32 InMouseX, int32 InMouseY) override;
+	virtual bool StartTracking(FLevelEditorViewportClient* InViewportClient, FViewport* InViewport) override;
+	virtual bool EndTracking(FLevelEditorViewportClient* InViewportClient, FViewport* InViewport) override;
 	// 	virtual void Tick(FLevelEditorViewportClient* ViewportClient,float DeltaTime) OVERRIDE;
-	virtual bool InputKey(FLevelEditorViewportClient* InViewportClient, FViewport* InViewport, FKey InKey, EInputEvent InEvent) OVERRIDE;
-	virtual bool InputDelta(FLevelEditorViewportClient* InViewportClient, FViewport* InViewport, FVector& InDrag, FRotator& InRot, FVector& InScale) OVERRIDE;
-	virtual void Render(const FSceneView* View, FViewport* Viewport, FPrimitiveDrawInterface* PDI) OVERRIDE;
-	virtual void DrawHUD(FLevelEditorViewportClient* ViewportClient, FViewport* Viewport, const FSceneView* View, FCanvas* Canvas) OVERRIDE;
+	virtual bool InputKey(FLevelEditorViewportClient* InViewportClient, FViewport* InViewport, FKey InKey, EInputEvent InEvent) override;
+	virtual bool InputDelta(FLevelEditorViewportClient* InViewportClient, FViewport* InViewport, FVector& InDrag, FRotator& InRot, FVector& InScale) override;
+	virtual void Render(const FSceneView* View, FViewport* Viewport, FPrimitiveDrawInterface* PDI) override;
+	virtual void DrawHUD(FLevelEditorViewportClient* ViewportClient, FViewport* Viewport, const FSceneView* View, FCanvas* Canvas) override;
 	// 	virtual bool Select(AActor* InActor, bool bInSelected) OVERRIDE;
 	// 	virtual bool IsSelectionAllowed(AActor* InActor) const OVERRIDE;
 	// 	virtual void ActorSelectionChangeNotify() OVERRIDE;
 	// 	virtual FVector GetWidgetLocation() const OVERRIDE;
 	virtual bool AllowWidgetMove();
-	virtual bool ShouldDrawWidget() const OVERRIDE;
-	virtual bool UsesTransformWidget() const OVERRIDE;
-	virtual void PeekAtSelectionChangedEvent(UObject* ItemUndergoingChange) OVERRIDE;
+	virtual bool ShouldDrawWidget() const override;
+	virtual bool UsesTransformWidget() const override;
+	virtual void PeekAtSelectionChangedEvent(UObject* ItemUndergoingChange) override;
 	// 	virtual int32 GetWidgetAxisToDraw(FWidget::EWidgetMode InWidgetMode) const OVERRIDE;
 	// 	virtual bool DisallowMouseDeltaTracking() const OVERRIDE;
 	// End of FEdMode interface
diff --git a/Engine/Plugins/Experimental/Paper2D/Source/Paper2DEditor/Private/TileMapEditing/PaperTileMapDetailsCustomization.h b/Engine/Plugins/Experimental/Paper2D/Source/Paper2DEditor/Private/TileMapEditing/PaperTileMapDetailsCustomization.h
--- a/Engine/Plugins/Experimental/Paper2D/Source/Paper2DEditor/Private/TileMapEditing/PaperTileMapDetailsCustomization.h
+++ b/Engine/Plugins/Experimental/Paper2D/Source/Paper2DEditor/Private/TileMapEditing/PaperTileMapDetailsCustomization.h
@@ -14,7 +14,7 @@ public:
 	static TSharedRef<IDetailCustomization> MakeInstance();
 
 	// IDetailCustomization interface
-	virtual void CustomizeDetails(IDetailLayoutBuilder& DetailLayout) OVERRIDE;
+	virtual void CustomizeDetails(IDetailLayoutBuilder& DetailLayout) override;
 	// End of IDetailCustomization interface
 
 private:
diff --git a/Engine/Plugins/Experimental/Paper2D/Source/Paper2DEditor/Private/TileMapEditing/STileLayerItem.cpp b/Engine/Plugins/Experimental/Paper2D/Source/Paper2DEditor/Private/TileMapEditing/STileLayerItem.cpp
--- a/Engine/Plugins/Experimental/Paper2D/Source/Paper2DEditor/Private/TileMapEditing/STileLayerItem.cpp
+++ b/Engine/Plugins/Experimental/Paper2D/Source/Paper2DEditor/Private/TileMapEditing/STileLayerItem.cpp
@@ -21,8 +21,8 @@ public:
 	void Construct(const FArguments& InArgs);
 
 	// SWidget interface
-	virtual FReply OnMouseButtonDoubleClick(const FGeometry& InMyGeometry, const FPointerEvent& InMouseEvent) OVERRIDE;
-	virtual bool SupportsKeyboardFocus() const OVERRIDE;
+	virtual FReply OnMouseButtonDoubleClick(const FGeometry& InMyGeometry, const FPointerEvent& InMouseEvent) override;
+	virtual bool SupportsKeyboardFocus() const override;
 	// End of SWidget interface
 protected:
 	bool bInEditMode;
diff --git a/Engine/Plugins/Experimental/Paper2D/Source/Paper2DEditor/Private/TileMapEditing/STileLayerItem.h b/Engine/Plugins/Experimental/Paper2D/Source/Paper2DEditor/Private/TileMapEditing/STileLayerItem.h
--- a/Engine/Plugins/Experimental/Paper2D/Source/Paper2DEditor/Private/TileMapEditing/STileLayerItem.h
+++ b/Engine/Plugins/Experimental/Paper2D/Source/Paper2DEditor/Private/TileMapEditing/STileLayerItem.h
@@ -14,7 +14,7 @@ public:
 	void Construct(const FArguments& InArgs, class UPaperTileLayer* InItem, const TSharedRef<STableViewBase>& OwnerTable);
 
 	// SMultiColumnTableRow<> interface
-	virtual TSharedRef<SWidget> GenerateWidgetForColumn(const FName& ColumnName) OVERRIDE;
+	virtual TSharedRef<SWidget> GenerateWidgetForColumn(const FName& ColumnName) override;
 	// End of SMultiColumnTableRow<> interface
 protected:
 	class UPaperTileLayer* MyLayer;
diff --git a/Engine/Plugins/Experimental/Paper2D/Source/Paper2DEditor/Private/TileMapEditing/TileMapEdModeToolkit.h b/Engine/Plugins/Experimental/Paper2D/Source/Paper2DEditor/Private/TileMapEditing/TileMapEdModeToolkit.h
--- a/Engine/Plugins/Experimental/Paper2D/Source/Paper2DEditor/Private/TileMapEditing/TileMapEdModeToolkit.h
+++ b/Engine/Plugins/Experimental/Paper2D/Source/Paper2DEditor/Private/TileMapEditing/TileMapEdModeToolkit.h
@@ -11,17 +11,17 @@ class FTileMapEdModeToolkit : public FModeToolkit
 {
 public:
 	// IToolkit interface
-	virtual void RegisterTabSpawners(const TSharedRef<class FTabManager>& TabManager) OVERRIDE;
-	virtual void UnregisterTabSpawners(const TSharedRef<class FTabManager>& TabManager) OVERRIDE;
-	virtual FName GetToolkitFName() const OVERRIDE;
-	virtual FText GetBaseToolkitName() const OVERRIDE;
-	virtual FText GetToolkitName() const OVERRIDE;
-	virtual class FEdMode* GetEditorMode() const OVERRIDE;
-	virtual TSharedPtr<class SWidget> GetInlineContent() const OVERRIDE;
+	virtual void RegisterTabSpawners(const TSharedRef<class FTabManager>& TabManager) override;
+	virtual void UnregisterTabSpawners(const TSharedRef<class FTabManager>& TabManager) override;
+	virtual FName GetToolkitFName() const override;
+	virtual FText GetBaseToolkitName() const override;
+	virtual FText GetToolkitName() const override;
+	virtual class FEdMode* GetEditorMode() const override;
+	virtual TSharedPtr<class SWidget> GetInlineContent() const override;
 	// End of IToolkit interface
 
 	// FModeToolkit interface
-	virtual void Init(const TSharedPtr<class IToolkitHost>& InitToolkitHost) OVERRIDE;
+	virtual void Init(const TSharedPtr<class IToolkitHost>& InitToolkitHost) override;
 	// End of FModeToolkit interface
 
 protected:
diff --git a/Engine/Plugins/Experimental/Paper2D/Source/Paper2DEditor/Private/TileSetEditor.cpp b/Engine/Plugins/Experimental/Paper2D/Source/Paper2DEditor/Private/TileSetEditor.cpp
--- a/Engine/Plugins/Experimental/Paper2D/Source/Paper2DEditor/Private/TileSetEditor.cpp
+++ b/Engine/Plugins/Experimental/Paper2D/Source/Paper2DEditor/Private/TileSetEditor.cpp
@@ -48,7 +48,7 @@ public:
 	}
 
 	// FViewportClient interface
-	virtual void Draw(FViewport* Viewport, FCanvas* Canvas) OVERRIDE;
+	virtual void Draw(FViewport* Viewport, FCanvas* Canvas) override;
 	// End of FViewportClient interface
 
 public:
@@ -231,7 +231,7 @@ public:
 		ViewMenuTooltip = LOCTEXT("TextureViewTabMenu_ToolTip", "Shows the viewport");
 	}
 
-	virtual TSharedRef<SWidget> CreateTabBody(const FWorkflowTabSpawnInfo& Info) const OVERRIDE
+	virtual TSharedRef<SWidget> CreateTabBody(const FWorkflowTabSpawnInfo& Info) const override
 	{
 		TSharedPtr<FTileSetEditor> TileSetEditorPtr = StaticCastSharedPtr<FTileSetEditor>(HostingApp.Pin());
 		//TSharedRef<FTileSetEditorViewportClient> ViewportClientPtr = MakeShareable(new FTileSetEditorViewportClient(TileSetEditorPtr));
@@ -262,12 +262,12 @@ public:
 	}
 
 	// SSingleObjectDetailsPanel interface
-	virtual UObject* GetObjectToObserve() const OVERRIDE
+	virtual UObject* GetObjectToObserve() const override
 	{
 		return TileSetEditorPtr.Pin()->GetTileSetBeingEdited();
 	}
 
-	virtual TSharedRef<SWidget> PopulateSlot(TSharedRef<SWidget> PropertyEditorWidget) OVERRIDE
+	virtual TSharedRef<SWidget> PopulateSlot(TSharedRef<SWidget> PropertyEditorWidget) override
 	{
 		return SNew(SVerticalBox)
 			+SVerticalBox::Slot()
@@ -297,7 +297,7 @@ public:
 		ViewMenuTooltip = LOCTEXT("DetailsTabMenu_ToolTip", "Shows the details panel");
 	}
 
-	virtual TSharedRef<SWidget> CreateTabBody(const FWorkflowTabSpawnInfo& Info) const OVERRIDE
+	virtual TSharedRef<SWidget> CreateTabBody(const FWorkflowTabSpawnInfo& Info) const override
 	{
 		TSharedPtr<FTileSetEditor> TileSetEditorPtr = StaticCastSharedPtr<FTileSetEditor>(HostingApp.Pin());
 		return SNew(STileSetPropertiesTabBody, TileSetEditorPtr);
@@ -313,7 +313,7 @@ public:
 	FBasicTileSetEditorMode(TSharedPtr<class FTileSetEditor> InTileSetEditor, FName InModeName);
 
 	// FApplicationMode interface
-	virtual void RegisterTabFactories(TSharedPtr<FTabManager> InTabManager) OVERRIDE;
+	virtual void RegisterTabFactories(TSharedPtr<FTabManager> InTabManager) override;
 	// End of FApplicationMode interface
 
 protected:
diff --git a/Engine/Plugins/Experimental/Paper2D/Source/Paper2DEditor/Private/TileSetEditor.h b/Engine/Plugins/Experimental/Paper2D/Source/Paper2DEditor/Private/TileSetEditor.h
--- a/Engine/Plugins/Experimental/Paper2D/Source/Paper2DEditor/Private/TileSetEditor.h
+++ b/Engine/Plugins/Experimental/Paper2D/Source/Paper2DEditor/Private/TileSetEditor.h
@@ -12,19 +12,19 @@ class FTileSetEditor : public FWorkflowCentricApplication, public FGCObject
 {
 public:
 	// IToolkit interface
-	virtual void RegisterTabSpawners(const TSharedRef<class FTabManager>& TabManager) OVERRIDE;
+	virtual void RegisterTabSpawners(const TSharedRef<class FTabManager>& TabManager) override;
 	// End of IToolkit interface
 
 	// FAssetEditorToolkit
-	virtual FName GetToolkitFName() const OVERRIDE;
-	virtual FText GetBaseToolkitName() const OVERRIDE;
-	virtual FText GetToolkitName() const OVERRIDE;
-	virtual FLinearColor GetWorldCentricTabColorScale() const OVERRIDE;
-	virtual FString GetWorldCentricTabPrefix() const OVERRIDE;
+	virtual FName GetToolkitFName() const override;
+	virtual FText GetBaseToolkitName() const override;
+	virtual FText GetToolkitName() const override;
+	virtual FLinearColor GetWorldCentricTabColorScale() const override;
+	virtual FString GetWorldCentricTabPrefix() const override;
 	// End of FAssetEditorToolkit
 
 	// FSerializableObject interface
-	virtual void AddReferencedObjects(FReferenceCollector& Collector) OVERRIDE;
+	virtual void AddReferencedObjects(FReferenceCollector& Collector) override;
 	// End of FSerializableObject interface
 
 public:
@@ -51,7 +51,7 @@ public:
 	void ChangeTileSet(UPaperTileSet* InTileSet);
 protected:
 	// SPaperEditorViewport interface
-	virtual FText GetTitleText() const OVERRIDE;
+	virtual FText GetTitleText() const override;
 	// End of SPaperEditorViewport interface
 
 private:
diff --git a/Engine/Plugins/Experimental/Paper2D/Source/PaperJsonImporter/Classes/PaperJsonImporterFactory.h b/Engine/Plugins/Experimental/Paper2D/Source/PaperJsonImporter/Classes/PaperJsonImporterFactory.h
--- a/Engine/Plugins/Experimental/Paper2D/Source/PaperJsonImporter/Classes/PaperJsonImporterFactory.h
+++ b/Engine/Plugins/Experimental/Paper2D/Source/PaperJsonImporter/Classes/PaperJsonImporterFactory.h
@@ -11,8 +11,8 @@ class UPaperJsonImporterFactory : public UFactory
 	GENERATED_UCLASS_BODY()
 
 	// UFactory interface
-	virtual FText GetToolTip() const OVERRIDE;
-	virtual UObject* FactoryCreateText(UClass* InClass, UObject* InParent, FName InName, EObjectFlags Flags, UObject* Context, const TCHAR* Type, const TCHAR*& Buffer, const TCHAR* BufferEnd, FFeedbackContext* Warn) OVERRIDE;
+	virtual FText GetToolTip() const override;
+	virtual UObject* FactoryCreateText(UClass* InClass, UObject* InParent, FName InName, EObjectFlags Flags, UObject* Context, const TCHAR* Type, const TCHAR*& Buffer, const TCHAR* BufferEnd, FFeedbackContext* Warn) override;
 	// End of UFactory interface
 
 protected:
diff --git a/Engine/Plugins/Experimental/Paper2D/Source/SmartSnapping/Private/PlanarConstraintSnapPolicy.h b/Engine/Plugins/Experimental/Paper2D/Source/SmartSnapping/Private/PlanarConstraintSnapPolicy.h
--- a/Engine/Plugins/Experimental/Paper2D/Source/SmartSnapping/Private/PlanarConstraintSnapPolicy.h
+++ b/Engine/Plugins/Experimental/Paper2D/Source/SmartSnapping/Private/PlanarConstraintSnapPolicy.h
@@ -19,10 +19,10 @@ public:
 	FPlanarConstraintSnapPolicy();
 
 	// ISnappingPolicy interface
-	virtual void SnapScale(FVector& Point, const FVector& GridBase) OVERRIDE;
-	virtual void SnapPointToGrid(FVector& Point, const FVector& GridBase) OVERRIDE;
-	virtual void SnapRotatorToGrid(FRotator& Rotation) OVERRIDE;
-	virtual void ClearSnappingHelpers(bool bClearImmediately) OVERRIDE;
-	virtual void DrawSnappingHelpers(const FSceneView* View, FPrimitiveDrawInterface* PDI) OVERRIDE;
+	virtual void SnapScale(FVector& Point, const FVector& GridBase) override;
+	virtual void SnapPointToGrid(FVector& Point, const FVector& GridBase) override;
+	virtual void SnapRotatorToGrid(FRotator& Rotation) override;
+	virtual void ClearSnappingHelpers(bool bClearImmediately) override;
+	virtual void DrawSnappingHelpers(const FSceneView* View, FPrimitiveDrawInterface* PDI) override;
 	// End of ISnappingPolicy interface
 };
diff --git a/Engine/Plugins/Experimental/Phya/Source/Phya/Classes/PhyaCollisionHandler.h b/Engine/Plugins/Experimental/Phya/Source/Phya/Classes/PhyaCollisionHandler.h
--- a/Engine/Plugins/Experimental/Phya/Source/Phya/Classes/PhyaCollisionHandler.h
+++ b/Engine/Plugins/Experimental/Phya/Source/Phya/Classes/PhyaCollisionHandler.h
@@ -31,8 +31,8 @@ class UPhyaCollisionHandler : public UPhysicsCollisionHandler
 	GENERATED_UCLASS_BODY()
 
 	// Begin UPhysicsCollisionHandler interface
-	virtual void InitCollisionHandler() OVERRIDE;
-	virtual void HandlePhysicsCollisions(const TArray<FCollisionNotifyInfo>& PendingCollisionNotifies) OVERRIDE;
+	virtual void InitCollisionHandler() override;
+	virtual void HandlePhysicsCollisions(const TArray<FCollisionNotifyInfo>& PendingCollisionNotifies) override;
 	// End UPhysicsCollisionHandler interface
 
 	void TestImpact();
diff --git a/Engine/Plugins/Messaging/UdpMessaging/Source/UdpMessaging/Private/UdpMessagingModule.cpp b/Engine/Plugins/Messaging/UdpMessaging/Source/UdpMessaging/Private/UdpMessagingModule.cpp
--- a/Engine/Plugins/Messaging/UdpMessaging/Source/UdpMessaging/Private/UdpMessagingModule.cpp
+++ b/Engine/Plugins/Messaging/UdpMessaging/Source/UdpMessaging/Private/UdpMessagingModule.cpp
@@ -20,7 +20,7 @@ public:
 
 	// Begin IModuleInterface interface
 
-	virtual void StartupModule( ) OVERRIDE
+	virtual void StartupModule( ) override
 	{
 		if (!SupportsNetworkedTransport())
 		{
@@ -58,7 +58,7 @@ public:
 		RestartServices();
 	}
 
-	virtual void ShutdownModule( ) OVERRIDE
+	virtual void ShutdownModule( ) override
 	{
 		// unregister application events
 		FCoreDelegates::ApplicationHasReactivatedDelegate.RemoveAll(this);
@@ -77,7 +77,7 @@ public:
 		ShutdownTunnel();
 	}
 
-	virtual bool SupportsDynamicReloading( ) OVERRIDE
+	virtual bool SupportsDynamicReloading( ) override
 	{
 		return true;
 	}
diff --git a/Engine/Plugins/OculusRift/Source/OculusRift/Private/OculusRiftHMD.cpp b/Engine/Plugins/OculusRift/Source/OculusRift/Private/OculusRiftHMD.cpp
--- a/Engine/Plugins/OculusRift/Source/OculusRift/Private/OculusRiftHMD.cpp
+++ b/Engine/Plugins/OculusRift/Source/OculusRift/Private/OculusRiftHMD.cpp
@@ -14,7 +14,7 @@
 class FOculusRiftPlugin : public IOculusRiftPlugin
 {
 	/** IHeadMountedDisplayModule implementation */
-	virtual TSharedPtr< class IHeadMountedDisplay > CreateHeadMountedDisplay() OVERRIDE;
+	virtual TSharedPtr< class IHeadMountedDisplay > CreateHeadMountedDisplay() override;
 };
 
 IMPLEMENT_MODULE( FOculusRiftPlugin, OculusRift )
diff --git a/Engine/Plugins/OculusRift/Source/OculusRift/Private/OculusRiftHMD.h b/Engine/Plugins/OculusRift/Source/OculusRift/Private/OculusRiftHMD.h
--- a/Engine/Plugins/OculusRift/Source/OculusRift/Private/OculusRiftHMD.h
+++ b/Engine/Plugins/OculusRift/Source/OculusRift/Private/OculusRiftHMD.h
@@ -57,67 +57,67 @@ class FOculusRiftHMD : public IHeadMountedDisplay, public ISceneViewExtension
 
 public:
 	/** IHeadMountedDisplay interface */
-	virtual bool IsHMDEnabled() const OVERRIDE;
-	virtual void EnableHMD(bool allow = true) OVERRIDE;
-	virtual EHMDDeviceType::Type GetHMDDeviceType() const OVERRIDE;
-	virtual bool GetHMDMonitorInfo(MonitorInfo&) const OVERRIDE;
+	virtual bool IsHMDEnabled() const override;
+	virtual void EnableHMD(bool allow = true) override;
+	virtual EHMDDeviceType::Type GetHMDDeviceType() const override;
+	virtual bool GetHMDMonitorInfo(MonitorInfo&) const override;
 
-	virtual bool DoesSupportPositionalTracking() const OVERRIDE;
-	virtual bool HasValidTrackingPosition() const OVERRIDE;
-	virtual void GetPositionalTrackingCameraProperties(FVector& OutOrigin, FRotator& OutOrientation, float& OutHFOV, float& OutVFOV, float& OutCameraDistance, float& OutNearPlane, float& OutFarPlane) const OVERRIDE;
+	virtual bool DoesSupportPositionalTracking() const override;
+	virtual bool HasValidTrackingPosition() const override;
+	virtual void GetPositionalTrackingCameraProperties(FVector& OutOrigin, FRotator& OutOrientation, float& OutHFOV, float& OutVFOV, float& OutCameraDistance, float& OutNearPlane, float& OutFarPlane) const override;
 
-	virtual void SetInterpupillaryDistance(float NewInterpupillaryDistance) OVERRIDE;
-	virtual float GetInterpupillaryDistance() const OVERRIDE;
+	virtual void SetInterpupillaryDistance(float NewInterpupillaryDistance) override;
+	virtual float GetInterpupillaryDistance() const override;
     //virtual float GetFieldOfViewInRadians() const OVERRIDE;
-	virtual void GetFieldOfView(float& OutHFOVInDegrees, float& OutVFOVInDegrees) const OVERRIDE;
+	virtual void GetFieldOfView(float& OutHFOVInDegrees, float& OutVFOVInDegrees) const override;
 
-    virtual void GetCurrentOrientationAndPosition(FQuat& CurrentOrientation, FVector& CurrentPosition) const OVERRIDE;
-	virtual void ApplyHmdRotation(APlayerController* PC, FRotator& ViewRotation) OVERRIDE;
-	virtual void UpdatePlayerCameraRotation(APlayerCameraManager*, struct FMinimalViewInfo& POV) OVERRIDE;
+    virtual void GetCurrentOrientationAndPosition(FQuat& CurrentOrientation, FVector& CurrentPosition) const override;
+	virtual void ApplyHmdRotation(APlayerController* PC, FRotator& ViewRotation) override;
+	virtual void UpdatePlayerCameraRotation(APlayerCameraManager*, struct FMinimalViewInfo& POV) override;
 
-	virtual bool IsChromaAbCorrectionEnabled() const OVERRIDE;
+	virtual bool IsChromaAbCorrectionEnabled() const override;
 
-	virtual class ISceneViewExtension* GetViewExtension() OVERRIDE;
-	virtual bool Exec( UWorld* InWorld, const TCHAR* Cmd, FOutputDevice& Ar ) OVERRIDE;
-	virtual void OnScreenModeChange(EWindowMode::Type WindowMode) OVERRIDE;
+	virtual class ISceneViewExtension* GetViewExtension() override;
+	virtual bool Exec( UWorld* InWorld, const TCHAR* Cmd, FOutputDevice& Ar ) override;
+	virtual void OnScreenModeChange(EWindowMode::Type WindowMode) override;
 
 	/** IStereoRendering interface */
-	virtual bool IsStereoEnabled() const OVERRIDE;
-	virtual bool EnableStereo(bool stereo = true) OVERRIDE;
-    virtual void AdjustViewRect(EStereoscopicPass StereoPass, int32& X, int32& Y, uint32& SizeX, uint32& SizeY) const OVERRIDE;
+	virtual bool IsStereoEnabled() const override;
+	virtual bool EnableStereo(bool stereo = true) override;
+    virtual void AdjustViewRect(EStereoscopicPass StereoPass, int32& X, int32& Y, uint32& SizeX, uint32& SizeY) const override;
 	virtual void CalculateStereoViewOffset(const EStereoscopicPass StereoPassType, const FRotator& ViewRotation, 
-										   const float MetersToWorld, FVector& ViewLocation) OVERRIDE;
-	virtual FMatrix GetStereoProjectionMatrix(const EStereoscopicPass StereoPassType, const float FOV) const OVERRIDE;
-	virtual void InitCanvasFromView(FSceneView* InView, UCanvas* Canvas) OVERRIDE;
-	virtual void PushViewportCanvas(EStereoscopicPass StereoPass, FCanvas *InCanvas, UCanvas *InCanvasObject, FViewport *InViewport) const OVERRIDE;
-	virtual void PushViewCanvas(EStereoscopicPass StereoPass, FCanvas *InCanvas, UCanvas *InCanvasObject, FSceneView *InView) const OVERRIDE;
-	virtual void GetEyeRenderParams_RenderThread(EStereoscopicPass StereoPass, FVector2D& EyeToSrcUVScaleValue, FVector2D& EyeToSrcUVOffsetValue) const OVERRIDE;
+										   const float MetersToWorld, FVector& ViewLocation) override;
+	virtual FMatrix GetStereoProjectionMatrix(const EStereoscopicPass StereoPassType, const float FOV) const override;
+	virtual void InitCanvasFromView(FSceneView* InView, UCanvas* Canvas) override;
+	virtual void PushViewportCanvas(EStereoscopicPass StereoPass, FCanvas *InCanvas, UCanvas *InCanvasObject, FViewport *InViewport) const override;
+	virtual void PushViewCanvas(EStereoscopicPass StereoPass, FCanvas *InCanvas, UCanvas *InCanvasObject, FSceneView *InView) const override;
+	virtual void GetEyeRenderParams_RenderThread(EStereoscopicPass StereoPass, FVector2D& EyeToSrcUVScaleValue, FVector2D& EyeToSrcUVOffsetValue) const override;
 
     /** ISceneViewExtension interface */
-    virtual void ModifyShowFlags(FEngineShowFlags& ShowFlags) OVERRIDE;
-    virtual void SetupView(FSceneViewFamily& InViewFamily, FSceneView& InView) OVERRIDE;
-    virtual void PreRenderView_RenderThread(FSceneView& InView) OVERRIDE;
-	virtual void PreRenderViewFamily_RenderThread(FSceneViewFamily& InViewFamily) OVERRIDE;
+    virtual void ModifyShowFlags(FEngineShowFlags& ShowFlags) override;
+    virtual void SetupView(FSceneViewFamily& InViewFamily, FSceneView& InView) override;
+    virtual void PreRenderView_RenderThread(FSceneView& InView) override;
+	virtual void PreRenderViewFamily_RenderThread(FSceneViewFamily& InViewFamily) override;
 
 	/** Positional tracking control methods */
-	virtual bool IsPositionalTrackingEnabled() const OVERRIDE;
-	virtual bool EnablePositionalTracking(bool enable) OVERRIDE;
+	virtual bool IsPositionalTrackingEnabled() const override;
+	virtual bool EnablePositionalTracking(bool enable) override;
 
 	/** A hookup for latency tester (render thread). */
-	virtual bool GetLatencyTesterColor_RenderThread(FColor& color, const FSceneView& view) OVERRIDE;
+	virtual bool GetLatencyTesterColor_RenderThread(FColor& color, const FSceneView& view) override;
 
-	virtual bool IsHeadTrackingAllowed() const OVERRIDE;
+	virtual bool IsHeadTrackingAllowed() const override;
 
-	virtual bool IsInLowPersistenceMode() const OVERRIDE;
-	virtual void EnableLowPersistenceMode(bool Enable = true) OVERRIDE;
+	virtual bool IsInLowPersistenceMode() const override;
+	virtual void EnableLowPersistenceMode(bool Enable = true) override;
 
 	/** Resets orientation by setting roll and pitch to 0, 
 	    assuming that current yaw is forward direction and assuming
 		current position as 0 point. */
-	virtual void ResetOrientationAndPosition(float yaw = 0.f) OVERRIDE;
+	virtual void ResetOrientationAndPosition(float yaw = 0.f) override;
 
-	virtual void DrawDistortionMesh_RenderThread(struct FRenderingCompositePassContext& Context, const FSceneView& View, const FIntPoint& TextureSize) OVERRIDE;
-	virtual void UpdateScreenSettings(const FViewport*) OVERRIDE;
+	virtual void DrawDistortionMesh_RenderThread(struct FRenderingCompositePassContext& Context, const FSceneView& View, const FIntPoint& TextureSize) override;
+	virtual void UpdateScreenSettings(const FViewport*) override;
 
 #ifdef OVR_DIRECT_RENDERING
 #if defined(OVR_D3D_VERSION) && (OVR_D3D_VERSION == 11)
@@ -145,16 +145,16 @@ public:
 		bool				bReadOnly;
 	};
 
-	virtual ID3D11Bridge* GetD3D11Bridge() OVERRIDE 
+	virtual ID3D11Bridge* GetD3D11Bridge() override 
 	{ 
 		check(IsInGameThread());
 		return &mD3D11Bridge; 
 	} 
 
-	virtual void UpdateRenderTarget(FRHITexture* TargetableTexture, FRHITexture* ShaderResourceTexture, uint32 SizeX, uint32 SizeY, int NumSamples) OVERRIDE;
-	virtual void CalculateRenderTargetSize(uint32& InOutSizeX, uint32& InOutSizeY) const OVERRIDE;
-	virtual bool NeedReAllocateRenderTarget(uint32 InSizeX, uint32 InSizeY) const OVERRIDE;
-	virtual bool ShouldUseSeparateRenderTarget() const OVERRIDE 
+	virtual void UpdateRenderTarget(FRHITexture* TargetableTexture, FRHITexture* ShaderResourceTexture, uint32 SizeX, uint32 SizeY, int NumSamples) override;
+	virtual void CalculateRenderTargetSize(uint32& InOutSizeX, uint32& InOutSizeY) const override;
+	virtual bool NeedReAllocateRenderTarget(uint32 InSizeX, uint32 InSizeY) const override;
+	virtual bool ShouldUseSeparateRenderTarget() const override 
 	{ 
 		check(IsInGameThread());
 		return IsStereoEnabled(); 
@@ -165,7 +165,7 @@ public:
 	void FinishRendering_RenderThread(); //!!AB
 #endif // #ifdef OVR_DIRECT_RENDERING
 
-	virtual void DrawDebug(UCanvas* Canvas) OVERRIDE;
+	virtual void DrawDebug(UCanvas* Canvas) override;
 
 #if 0 // !UE_BUILD_SHIPPING
     /** Debugging functionality */
diff --git a/Engine/Plugins/Runtime/Advertising/IOSTapJoy/Source/IOSTapJoy/Private/IOSTapJoy.cpp b/Engine/Plugins/Runtime/Advertising/IOSTapJoy/Source/IOSTapJoy/Private/IOSTapJoy.cpp
--- a/Engine/Plugins/Runtime/Advertising/IOSTapJoy/Source/IOSTapJoy/Private/IOSTapJoy.cpp
+++ b/Engine/Plugins/Runtime/Advertising/IOSTapJoy/Source/IOSTapJoy/Private/IOSTapJoy.cpp
@@ -7,13 +7,13 @@
 class FTapJoyProvider : public IAdvertisingProvider
 {
 	/** IModuleInterface implementation */
-	virtual void StartupModule() OVERRIDE;
-	virtual void ShutdownModule() OVERRIDE;
+	virtual void StartupModule() override;
+	virtual void ShutdownModule() override;
 
 	/** IAdvertisingProvider implementation */
-	virtual void ShowAdBanner( bool bShowOnBottomOfScreen ) OVERRIDE;
-	virtual void HideAdBanner() OVERRIDE;
-	virtual void CloseAdBanner() OVERRIDE;
+	virtual void ShowAdBanner( bool bShowOnBottomOfScreen ) override;
+	virtual void HideAdBanner() override;
+	virtual void CloseAdBanner() override;
 };
 
 IMPLEMENT_MODULE( FTapJoyProvider, IOSTapJoy )
diff --git a/Engine/Plugins/Runtime/AndroidDeviceProfileSelector/Source/AndroidDeviceProfileSelector/Private/AndroidDeviceProfileSelectorModule.h b/Engine/Plugins/Runtime/AndroidDeviceProfileSelector/Source/AndroidDeviceProfileSelector/Private/AndroidDeviceProfileSelectorModule.h
--- a/Engine/Plugins/Runtime/AndroidDeviceProfileSelector/Source/AndroidDeviceProfileSelector/Private/AndroidDeviceProfileSelectorModule.h
+++ b/Engine/Plugins/Runtime/AndroidDeviceProfileSelector/Source/AndroidDeviceProfileSelector/Private/AndroidDeviceProfileSelectorModule.h
@@ -12,13 +12,13 @@ class FAndroidDeviceProfileSelectorModule
 public:
 
 	// Begin IDeviceProfileSelectorModule interface
-	virtual const FString GetRuntimeDeviceProfileName() OVERRIDE;
+	virtual const FString GetRuntimeDeviceProfileName() override;
 	// End IDeviceProfileSelectorModule interface
 
 
 	// Begin IModuleInterface interface
-	virtual void StartupModule() OVERRIDE;
-	virtual void ShutdownModule() OVERRIDE;
+	virtual void StartupModule() override;
+	virtual void ShutdownModule() override;
 	// End IModuleInterface interface
 
 	
diff --git a/Engine/Plugins/Runtime/AppleMoviePlayer/Source/AppleMoviePlayer/Private/AppleMovieStreamer.h b/Engine/Plugins/Runtime/AppleMoviePlayer/Source/AppleMoviePlayer/Private/AppleMovieStreamer.h
--- a/Engine/Plugins/Runtime/AppleMoviePlayer/Source/AppleMoviePlayer/Private/AppleMovieStreamer.h
+++ b/Engine/Plugins/Runtime/AppleMoviePlayer/Source/AppleMoviePlayer/Private/AppleMovieStreamer.h
@@ -17,12 +17,12 @@ public:
 	FAVPlayerMovieStreamer();
 	virtual ~FAVPlayerMovieStreamer();
 
-	virtual void Init(const TArray<FString>& MoviePaths) OVERRIDE;
-	virtual void ForceCompletion() OVERRIDE;
-	virtual bool Tick(float DeltaTime) OVERRIDE;
-	virtual TSharedPtr<class ISlateViewport> GetViewportInterface() OVERRIDE;
-	virtual float GetAspectRatio() const OVERRIDE;
-	virtual void Cleanup() OVERRIDE;
+	virtual void Init(const TArray<FString>& MoviePaths) override;
+	virtual void ForceCompletion() override;
+	virtual bool Tick(float DeltaTime) override;
+	virtual TSharedPtr<class ISlateViewport> GetViewportInterface() override;
+	virtual float GetAspectRatio() const override;
+	virtual void Cleanup() override;
 
 private:
 
diff --git a/Engine/Plugins/Runtime/CableComponent/Source/CableComponent/Classes/CableComponent.h b/Engine/Plugins/Runtime/CableComponent/Source/CableComponent/Classes/CableComponent.h
--- a/Engine/Plugins/Runtime/CableComponent/Source/CableComponent/Classes/CableComponent.h
+++ b/Engine/Plugins/Runtime/CableComponent/Source/CableComponent/Classes/CableComponent.h
@@ -30,21 +30,21 @@ class UCableComponent : public UMeshComponent
 public:
 
 	// Begin UActorComponent interface.
-	virtual void OnRegister() OVERRIDE;
-	virtual void TickComponent(float DeltaTime, enum ELevelTick TickType, FActorComponentTickFunction *ThisTickFunction) OVERRIDE;
-	virtual void SendRenderDynamicData_Concurrent() OVERRIDE;
+	virtual void OnRegister() override;
+	virtual void TickComponent(float DeltaTime, enum ELevelTick TickType, FActorComponentTickFunction *ThisTickFunction) override;
+	virtual void SendRenderDynamicData_Concurrent() override;
 	// Begin UActorComponent interface.
 
 	// Begin USceneComponent interface.
-	virtual FBoxSphereBounds CalcBounds(const FTransform & LocalToWorld) const OVERRIDE;
+	virtual FBoxSphereBounds CalcBounds(const FTransform & LocalToWorld) const override;
 	// Begin USceneComponent interface.
 
 	// Begin UPrimitiveComponent interface.
-	virtual FPrimitiveSceneProxy* CreateSceneProxy() OVERRIDE;
+	virtual FPrimitiveSceneProxy* CreateSceneProxy() override;
 	// End UPrimitiveComponent interface.
 
 	// Begin UMeshComponent interface.
-	virtual int32 GetNumMaterials() const OVERRIDE;
+	virtual int32 GetNumMaterials() const override;
 	// End UMeshComponent interface.
 
 
diff --git a/Engine/Plugins/Runtime/CustomMeshComponent/Source/CustomMeshComponent/Classes/CustomMeshComponent.h b/Engine/Plugins/Runtime/CustomMeshComponent/Source/CustomMeshComponent/Classes/CustomMeshComponent.h
--- a/Engine/Plugins/Runtime/CustomMeshComponent/Source/CustomMeshComponent/Classes/CustomMeshComponent.h
+++ b/Engine/Plugins/Runtime/CustomMeshComponent/Source/CustomMeshComponent/Classes/CustomMeshComponent.h
@@ -32,15 +32,15 @@ class UCustomMeshComponent : public UMeshComponent
 private:
 
 	// Begin UPrimitiveComponent interface.
-	virtual FPrimitiveSceneProxy* CreateSceneProxy() OVERRIDE;
+	virtual FPrimitiveSceneProxy* CreateSceneProxy() override;
 	// End UPrimitiveComponent interface.
 
 	// Begin UMeshComponent interface.
-	virtual int32 GetNumMaterials() const OVERRIDE;
+	virtual int32 GetNumMaterials() const override;
 	// End UMeshComponent interface.
 
 	// Begin USceneComponent interface.
-	virtual FBoxSphereBounds CalcBounds(const FTransform & LocalToWorld) const OVERRIDE;
+	virtual FBoxSphereBounds CalcBounds(const FTransform & LocalToWorld) const override;
 	// Begin USceneComponent interface.
 
 	/** */
diff --git a/Engine/Plugins/Runtime/CustomMeshComponent/Source/CustomMeshComponent/Private/CustomMeshComponent.cpp b/Engine/Plugins/Runtime/CustomMeshComponent/Source/CustomMeshComponent/Private/CustomMeshComponent.cpp
--- a/Engine/Plugins/Runtime/CustomMeshComponent/Source/CustomMeshComponent/Private/CustomMeshComponent.cpp
+++ b/Engine/Plugins/Runtime/CustomMeshComponent/Source/CustomMeshComponent/Private/CustomMeshComponent.cpp
@@ -191,7 +191,7 @@ public:
 		return Result;
 	}
 
-	virtual bool CanBeOccluded() const OVERRIDE
+	virtual bool CanBeOccluded() const override
 	{
 		return !MaterialRelevance.bDisableDepthTest;
 	}
diff --git a/Engine/Plugins/Runtime/ExampleDeviceProfileSelector/Source/ExampleDeviceProfileSelector/Private/ExampleDeviceProfileSelectorModule.h b/Engine/Plugins/Runtime/ExampleDeviceProfileSelector/Source/ExampleDeviceProfileSelector/Private/ExampleDeviceProfileSelectorModule.h
--- a/Engine/Plugins/Runtime/ExampleDeviceProfileSelector/Source/ExampleDeviceProfileSelector/Private/ExampleDeviceProfileSelectorModule.h
+++ b/Engine/Plugins/Runtime/ExampleDeviceProfileSelector/Source/ExampleDeviceProfileSelector/Private/ExampleDeviceProfileSelectorModule.h
@@ -12,13 +12,13 @@ class FExampleDeviceProfileSelectorModule
 public:
 
 	// Begin IDeviceProfileSelectorModule interface
-	virtual const FString GetRuntimeDeviceProfileName() OVERRIDE;
+	virtual const FString GetRuntimeDeviceProfileName() override;
 	// End IDeviceProfileSelectorModule interface
 
 
 	// Begin IModuleInterface interface
-	virtual void StartupModule() OVERRIDE;
-	virtual void ShutdownModule() OVERRIDE;
+	virtual void StartupModule() override;
+	virtual void ShutdownModule() override;
 	// End IModuleInterface interface
 
 	
diff --git a/Engine/Plugins/Runtime/LiveEditorListenServer/Source/LiveEditorListenServer/Private/LiveEditorListenServerEngineHooks.h b/Engine/Plugins/Runtime/LiveEditorListenServer/Source/LiveEditorListenServer/Private/LiveEditorListenServerEngineHooks.h
--- a/Engine/Plugins/Runtime/LiveEditorListenServer/Source/LiveEditorListenServer/Private/LiveEditorListenServerEngineHooks.h
+++ b/Engine/Plugins/Runtime/LiveEditorListenServer/Source/LiveEditorListenServer/Private/LiveEditorListenServerEngineHooks.h
@@ -48,7 +48,7 @@ public:
 	{
 		return false;
 	}
-	virtual TStatId GetStatId() const OVERRIDE;
+	virtual TStatId GetStatId() const override;
 
 private:
 	FLiveEditorListenServer *Owner;
diff --git a/Engine/Plugins/Runtime/WindowsMoviePlayer/Source/WindowsMoviePlayer/Private/WindowsMovieStreamer.h b/Engine/Plugins/Runtime/WindowsMoviePlayer/Source/WindowsMoviePlayer/Private/WindowsMovieStreamer.h
--- a/Engine/Plugins/Runtime/WindowsMoviePlayer/Source/WindowsMoviePlayer/Private/WindowsMovieStreamer.h
+++ b/Engine/Plugins/Runtime/WindowsMoviePlayer/Source/WindowsMoviePlayer/Private/WindowsMovieStreamer.h
@@ -20,18 +20,18 @@ public:
 	~FMediaFoundationMovieStreamer();
 
 	/** IMovieStreamer interface */
-	virtual void Init(const TArray<FString>& MoviePaths) OVERRIDE;
-	virtual void ForceCompletion() OVERRIDE;
-	virtual bool Tick(float DeltaTime) OVERRIDE;
-	virtual TSharedPtr<class ISlateViewport> GetViewportInterface() OVERRIDE
+	virtual void Init(const TArray<FString>& MoviePaths) override;
+	virtual void ForceCompletion() override;
+	virtual bool Tick(float DeltaTime) override;
+	virtual TSharedPtr<class ISlateViewport> GetViewportInterface() override
 	{
 		return MovieViewport;
 	}
-	virtual float GetAspectRatio() const OVERRIDE
+	virtual float GetAspectRatio() const override
 	{
 		return (float)MovieViewport->GetSize().X / (float)MovieViewport->GetSize().Y;
 	}
-	virtual void Cleanup() OVERRIDE;
+	virtual void Cleanup() override;
 
 private:
 	/** Opens up the next movie in the movie path queue */
diff --git a/Engine/Plugins/Script/ScriptEditorPlugin/Source/ScriptEditorPlugin/Classes/ReimportScriptFactory.h b/Engine/Plugins/Script/ScriptEditorPlugin/Source/ScriptEditorPlugin/Classes/ReimportScriptFactory.h
--- a/Engine/Plugins/Script/ScriptEditorPlugin/Source/ScriptEditorPlugin/Classes/ReimportScriptFactory.h
+++ b/Engine/Plugins/Script/ScriptEditorPlugin/Source/ScriptEditorPlugin/Classes/ReimportScriptFactory.h
@@ -17,9 +17,9 @@ class UReimportScriptFactory : public UScriptFactory, public FReimportHandler
 	class UScriptAsset* OriginalScript;
 
 	// Begin FReimportHandler interface
-	virtual bool CanReimport(UObject* Obj, TArray<FString>& OutFilenames) OVERRIDE;
-	virtual void SetReimportPaths(UObject* Obj, const TArray<FString>& NewReimportPaths) OVERRIDE;
-	virtual EReimportResult::Type Reimport(UObject* Obj) OVERRIDE;
+	virtual bool CanReimport(UObject* Obj, TArray<FString>& OutFilenames) override;
+	virtual void SetReimportPaths(UObject* Obj, const TArray<FString>& NewReimportPaths) override;
+	virtual EReimportResult::Type Reimport(UObject* Obj) override;
 	// End FReimportHandler interface
 };
 
diff --git a/Engine/Plugins/Script/ScriptEditorPlugin/Source/ScriptEditorPlugin/Classes/ScriptFactory.h b/Engine/Plugins/Script/ScriptEditorPlugin/Source/ScriptEditorPlugin/Classes/ScriptFactory.h
--- a/Engine/Plugins/Script/ScriptEditorPlugin/Source/ScriptEditorPlugin/Classes/ScriptFactory.h
+++ b/Engine/Plugins/Script/ScriptEditorPlugin/Source/ScriptEditorPlugin/Classes/ScriptFactory.h
@@ -12,11 +12,11 @@ class SCRIPTEDITORPLUGIN_API UScriptFactory : public UFactory
 public:
 
 	// Begin UObject Interface
-	virtual void PostInitProperties() OVERRIDE;
+	virtual void PostInitProperties() override;
 	// End UObject Interface
 
 	// Begin UFactory Interface
-	virtual bool DoesSupportClass(UClass* Class) OVERRIDE;
+	virtual bool DoesSupportClass(UClass* Class) override;
 	virtual UObject* FactoryCreateText(UClass* InClass, UObject* InParent, FName InName, EObjectFlags Flags, UObject* Context, const TCHAR* Type, const TCHAR*& Buffer, const TCHAR* BufferEnd, FFeedbackContext* Warn);
 	// End UFactory Interface
 };
diff --git a/Engine/Plugins/Script/ScriptGeneratorPlugin/Source/ScriptGeneratorPlugin/Private/GenericScriptCodeGenerator.h b/Engine/Plugins/Script/ScriptGeneratorPlugin/Source/ScriptGeneratorPlugin/Private/GenericScriptCodeGenerator.h
--- a/Engine/Plugins/Script/ScriptGeneratorPlugin/Source/ScriptGeneratorPlugin/Private/GenericScriptCodeGenerator.h
+++ b/Engine/Plugins/Script/ScriptGeneratorPlugin/Source/ScriptGeneratorPlugin/Private/GenericScriptCodeGenerator.h
@@ -44,6 +44,6 @@ public:
 	FGenericScriptCodeGenerator(const FString& RootLocalPath, const FString& RootBuildPath, const FString& OutputDirectory);
 
 	// FScriptCodeGeneratorBase interface
-	virtual void ExportClass(UClass* Class, const FString& SourceHeaderFilename, const FString& GeneratedHeaderFilename, bool bHasChanged) OVERRIDE;
-	virtual void FinishExport() OVERRIDE;
+	virtual void ExportClass(UClass* Class, const FString& SourceHeaderFilename, const FString& GeneratedHeaderFilename, bool bHasChanged) override;
+	virtual void FinishExport() override;
 };
diff --git a/Engine/Plugins/Script/ScriptGeneratorPlugin/Source/ScriptGeneratorPlugin/Private/LuaScriptCodeGenerator.h b/Engine/Plugins/Script/ScriptGeneratorPlugin/Source/ScriptGeneratorPlugin/Private/LuaScriptCodeGenerator.h
--- a/Engine/Plugins/Script/ScriptGeneratorPlugin/Source/ScriptGeneratorPlugin/Private/LuaScriptCodeGenerator.h
+++ b/Engine/Plugins/Script/ScriptGeneratorPlugin/Source/ScriptGeneratorPlugin/Private/LuaScriptCodeGenerator.h
@@ -48,16 +48,16 @@ protected:
 	bool IsPropertyTypeSupported(UProperty* Property) const;
 
 	// FScriptCodeGeneratorBase interface
-	virtual bool CanExportClass(UClass* Class) OVERRIDE;
-	virtual bool CanExportFunction(const FString& ClassNameCPP, UClass* Class, UFunction* Function) OVERRIDE;
-	virtual bool CanExportProperty(const FString& ClassNameCPP, UClass* Class, UProperty* Property) OVERRIDE;
-	virtual FString InitializeFunctionDispatchParam(UFunction* Function, UProperty* Param, int32 ParamIndex) OVERRIDE;
+	virtual bool CanExportClass(UClass* Class) override;
+	virtual bool CanExportFunction(const FString& ClassNameCPP, UClass* Class, UFunction* Function) override;
+	virtual bool CanExportProperty(const FString& ClassNameCPP, UClass* Class, UProperty* Property) override;
+	virtual FString InitializeFunctionDispatchParam(UFunction* Function, UProperty* Param, int32 ParamIndex) override;
 
 public:
 
 	FLuaScriptCodeGenerator(const FString& RootLocalPath, const FString& RootBuildPath, const FString& OutputDirectory);
 
 	// FScriptCodeGeneratorBase interface
-	virtual void ExportClass(UClass* Class, const FString& SourceHeaderFilename, const FString& GeneratedHeaderFilename, bool bHasChanged) OVERRIDE;
-	virtual void FinishExport() OVERRIDE;
+	virtual void ExportClass(UClass* Class, const FString& SourceHeaderFilename, const FString& GeneratedHeaderFilename, bool bHasChanged) override;
+	virtual void FinishExport() override;
 };
diff --git a/Engine/Plugins/Script/ScriptGeneratorPlugin/Source/ScriptGeneratorPlugin/Private/ScriptGeneratorPlugin.cpp b/Engine/Plugins/Script/ScriptGeneratorPlugin/Source/ScriptGeneratorPlugin/Private/ScriptGeneratorPlugin.cpp
--- a/Engine/Plugins/Script/ScriptGeneratorPlugin/Source/ScriptGeneratorPlugin/Private/ScriptGeneratorPlugin.cpp
+++ b/Engine/Plugins/Script/ScriptGeneratorPlugin/Source/ScriptGeneratorPlugin/Private/ScriptGeneratorPlugin.cpp
@@ -13,16 +13,16 @@ class FScriptGeneratorPlugin : public IScriptGeneratorPlugin
 	TAutoPtr<FScriptCodeGeneratorBase> CodeGenerator;
 
 	/** IModuleInterface implementation */
-	virtual void StartupModule() OVERRIDE;
-	virtual void ShutdownModule() OVERRIDE;
+	virtual void StartupModule() override;
+	virtual void ShutdownModule() override;
 
 	/** IScriptGeneratorPlugin interface */
-	virtual FString GetGeneratedCodeModuleName() const OVERRIDE { return TEXT("ScriptPlugin"); }
+	virtual FString GetGeneratedCodeModuleName() const override { return TEXT("ScriptPlugin"); }
 	virtual bool ShouldExportClassesForModule(const FString& ModuleName, EBuildModuleType::Type ModuleType) const;
-	virtual bool SupportsTarget(const FString& TargetName) const OVERRIDE { return true; }
-	virtual void Initialize(const FString& RootLocalPath, const FString& RootBuildPath, const FString& OutputDirectory) OVERRIDE;
-	virtual void ExportClass(UClass* Class, const FString& SourceHeaderFilename, const FString& GeneratedHeaderFilename, bool bHasChanged) OVERRIDE;
-	virtual void FinishExport() OVERRIDE;
+	virtual bool SupportsTarget(const FString& TargetName) const override { return true; }
+	virtual void Initialize(const FString& RootLocalPath, const FString& RootBuildPath, const FString& OutputDirectory) override;
+	virtual void ExportClass(UClass* Class, const FString& SourceHeaderFilename, const FString& GeneratedHeaderFilename, bool bHasChanged) override;
+	virtual void FinishExport() override;
 };
 
 IMPLEMENT_MODULE( FScriptGeneratorPlugin, ScriptGeneratorPlugin )
diff --git a/Engine/Plugins/Script/ScriptPlugin/Source/ScriptPlugin/Classes/ScriptAsset.h b/Engine/Plugins/Script/ScriptPlugin/Source/ScriptPlugin/Classes/ScriptAsset.h
--- a/Engine/Plugins/Script/ScriptPlugin/Source/ScriptPlugin/Classes/ScriptAsset.h
+++ b/Engine/Plugins/Script/ScriptPlugin/Source/ScriptPlugin/Classes/ScriptAsset.h
@@ -25,7 +25,7 @@ class SCRIPTPLUGIN_API UScriptAsset : public UDataAsset
 	FString SourceCode;
 
 	// UObject interface
-	virtual void PostLoad() OVERRIDE;
+	virtual void PostLoad() override;
 
 #if WITH_EDITOR
 	virtual void PostEditChangeProperty(struct FPropertyChangedEvent& PropertyChangedEvent);
diff --git a/Engine/Plugins/Script/ScriptPlugin/Source/ScriptPlugin/Classes/ScriptComponent.h b/Engine/Plugins/Script/ScriptPlugin/Source/ScriptPlugin/Classes/ScriptComponent.h
--- a/Engine/Plugins/Script/ScriptPlugin/Source/ScriptPlugin/Classes/ScriptComponent.h
+++ b/Engine/Plugins/Script/ScriptPlugin/Source/ScriptPlugin/Classes/ScriptComponent.h
@@ -29,10 +29,10 @@ public:
 	UScriptAsset* Script;
 
 	// Begin UActorComponent interface.
-	virtual void OnRegister() OVERRIDE;
-	virtual void InitializeComponent() OVERRIDE;
-	virtual void TickComponent(float DeltaTime, enum ELevelTick TickType, FActorComponentTickFunction *ThisTickFunction) OVERRIDE;
-	virtual void OnUnregister() OVERRIDE;
+	virtual void OnRegister() override;
+	virtual void InitializeComponent() override;
+	virtual void TickComponent(float DeltaTime, enum ELevelTick TickType, FActorComponentTickFunction *ThisTickFunction) override;
+	virtual void OnUnregister() override;
 	// Begin UActorComponent interface.
 
 protected:
diff --git a/Engine/Plugins/Script/ScriptPlugin/Source/ScriptPlugin/Private/LuaIntegration.h b/Engine/Plugins/Script/ScriptPlugin/Source/ScriptPlugin/Private/LuaIntegration.h
--- a/Engine/Plugins/Script/ScriptPlugin/Source/ScriptPlugin/Private/LuaIntegration.h
+++ b/Engine/Plugins/Script/ScriptPlugin/Source/ScriptPlugin/Private/LuaIntegration.h
@@ -22,11 +22,11 @@ public:
 	FLuaContext();
 
 	// FScriptContextBase interface
-	virtual bool Initialize(UScriptComponent* Owner) OVERRIDE;
-	virtual void BeginPlay() OVERRIDE;
-	virtual void Tick(float DeltaTime) OVERRIDE;
-	virtual void Destroy() OVERRIDE;
-	virtual bool CanTick() OVERRIDE;
+	virtual bool Initialize(UScriptComponent* Owner) override;
+	virtual void BeginPlay() override;
+	virtual void Tick(float DeltaTime) override;
+	virtual void Destroy() override;
+	virtual bool CanTick() override;
 };
 
 struct FLuaUtils
diff --git a/Engine/Plugins/Script/ScriptPlugin/Source/ScriptPlugin/Private/ScriptObjectReferencer.h b/Engine/Plugins/Script/ScriptPlugin/Source/ScriptPlugin/Private/ScriptObjectReferencer.h
--- a/Engine/Plugins/Script/ScriptPlugin/Source/ScriptPlugin/Private/ScriptObjectReferencer.h
+++ b/Engine/Plugins/Script/ScriptPlugin/Source/ScriptPlugin/Private/ScriptObjectReferencer.h
@@ -35,5 +35,5 @@ public:
 	void RemoveObjectReference(UObject* Object);
 
 	// FGCObject Interface
-	virtual void AddReferencedObjects(FReferenceCollector& Collector) OVERRIDE;
+	virtual void AddReferencedObjects(FReferenceCollector& Collector) override;
 };
diff --git a/Engine/Source/Developer/Android/AndroidDeviceDetection/Private/AndroidDeviceDetectionModule.cpp b/Engine/Source/Developer/Android/AndroidDeviceDetection/Private/AndroidDeviceDetectionModule.cpp
--- a/Engine/Source/Developer/Android/AndroidDeviceDetection/Private/AndroidDeviceDetectionModule.cpp
+++ b/Engine/Source/Developer/Android/AndroidDeviceDetection/Private/AndroidDeviceDetectionModule.cpp
@@ -259,12 +259,12 @@ public:
 		}
 	}
 
-	virtual const TMap<FString,FAndroidDeviceInfo>& GetDeviceMap() OVERRIDE
+	virtual const TMap<FString,FAndroidDeviceInfo>& GetDeviceMap() override
 	{
 		return DeviceMap;
 	}
 
-	virtual FCriticalSection* GetDeviceMapLock() OVERRIDE
+	virtual FCriticalSection* GetDeviceMapLock() override
 	{
 		return &DeviceMapLock;
 	}
@@ -305,7 +305,7 @@ public:
 		AndroidDeviceDetectionSingleton = nullptr;
 	}
 
-	virtual IAndroidDeviceDetection* GetAndroidDeviceDetection() OVERRIDE
+	virtual IAndroidDeviceDetection* GetAndroidDeviceDetection() override
 	{
 		if (AndroidDeviceDetectionSingleton == nullptr)
 		{
diff --git a/Engine/Source/Developer/Android/AndroidTargetPlatform/Private/AndroidTargetDevice.h b/Engine/Source/Developer/Android/AndroidTargetPlatform/Private/AndroidTargetDevice.h
--- a/Engine/Source/Developer/Android/AndroidTargetPlatform/Private/AndroidTargetDevice.h
+++ b/Engine/Source/Developer/Android/AndroidTargetPlatform/Private/AndroidTargetDevice.h
@@ -78,27 +78,27 @@ public:
 
 	// Begin ITargetDevice interface
 
-	virtual bool Connect( ) OVERRIDE
+	virtual bool Connect( ) override
 	{
 		return true;
 	}
 
-	virtual bool Deploy( const FString& SourceFolder, FString& OutAppId ) OVERRIDE;
+	virtual bool Deploy( const FString& SourceFolder, FString& OutAppId ) override;
 
-	virtual void Disconnect( ) OVERRIDE
+	virtual void Disconnect( ) override
 	{ }
 
-	virtual ETargetDeviceTypes::Type GetDeviceType( ) const OVERRIDE
+	virtual ETargetDeviceTypes::Type GetDeviceType( ) const override
 	{
 		return ETargetDeviceTypes::Tablet;
 	}
 
-	virtual FTargetDeviceId GetId( ) const OVERRIDE
+	virtual FTargetDeviceId GetId( ) const override
 	{
 		return FTargetDeviceId(TargetPlatform.PlatformName(), SerialNumber);
 	}
 
-	virtual FString GetName( ) const OVERRIDE
+	virtual FString GetName( ) const override
 	{
 		if (AndroidVariant.IsEmpty())
 		{
@@ -110,50 +110,50 @@ public:
 		return FString::Printf(TEXT("%s (%s)"), *Model, *AndroidVariant);
 	}
 
-	virtual FString GetOperatingSystemName( ) OVERRIDE
+	virtual FString GetOperatingSystemName( ) override
 	{
 		return TEXT("Android");
 	}
 
-	virtual int32 GetProcessSnapshot( TArray<FTargetDeviceProcessInfo>& OutProcessInfos ) OVERRIDE;
+	virtual int32 GetProcessSnapshot( TArray<FTargetDeviceProcessInfo>& OutProcessInfos ) override;
 
-	virtual const class ITargetPlatform& GetTargetPlatform( ) const OVERRIDE
+	virtual const class ITargetPlatform& GetTargetPlatform( ) const override
 	{
 		return TargetPlatform;
 	}
 
-	virtual bool IsConnected( ) OVERRIDE
+	virtual bool IsConnected( ) override
 	{
 		return Connected;
 	}
 
-	virtual bool IsDefault( ) const OVERRIDE
+	virtual bool IsDefault( ) const override
 	{
 		return true;
 	}
 
-	virtual bool PowerOff( bool Force ) OVERRIDE;
+	virtual bool PowerOff( bool Force ) override;
 
-	virtual bool PowerOn( ) OVERRIDE
+	virtual bool PowerOn( ) override
 	{
 		return true;
 	}
 
 	virtual bool Launch( const FString& AppId, EBuildConfigurations::Type BuildConfiguration, EBuildTargets::Type BuildTarget, const FString& Params, uint32* OutProcessId );
 
-	virtual bool Reboot( bool bReconnect = false ) OVERRIDE;
+	virtual bool Reboot( bool bReconnect = false ) override;
 
-	virtual bool Run( const FString& ExecutablePath, const FString& Params, uint32* OutProcessId ) OVERRIDE;
+	virtual bool Run( const FString& ExecutablePath, const FString& Params, uint32* OutProcessId ) override;
 
-	virtual bool SupportsFeature( ETargetDeviceFeatures::Type Feature ) const OVERRIDE;
+	virtual bool SupportsFeature( ETargetDeviceFeatures::Type Feature ) const override;
 
-	virtual bool SupportsSdkVersion( const FString& VersionString ) const OVERRIDE;
+	virtual bool SupportsSdkVersion( const FString& VersionString ) const override;
 
-	virtual bool TerminateProcess( const int32 ProcessId ) OVERRIDE;
+	virtual bool TerminateProcess( const int32 ProcessId ) override;
 
-	virtual void SetUserCredentials( const FString & UserName, const FString & UserPassword ) OVERRIDE;
+	virtual void SetUserCredentials( const FString & UserName, const FString & UserPassword ) override;
 
-	virtual bool GetUserCredentials( FString & OutUserName, FString & OutUserPassword ) OVERRIDE;
+	virtual bool GetUserCredentials( FString & OutUserName, FString & OutUserPassword ) override;
 
 	// End ITargetDevice interface
 
diff --git a/Engine/Source/Developer/Android/AndroidTargetPlatform/Private/AndroidTargetPlatform.h b/Engine/Source/Developer/Android/AndroidTargetPlatform/Private/AndroidTargetPlatform.h
--- a/Engine/Source/Developer/Android/AndroidTargetPlatform/Private/AndroidTargetPlatform.h
+++ b/Engine/Source/Developer/Android/AndroidTargetPlatform/Private/AndroidTargetPlatform.h
@@ -76,36 +76,36 @@ public:
 
 	// Begin ITargetPlatform interface
 
-	virtual void EnableDeviceCheck(bool OnOff) OVERRIDE {}
+	virtual void EnableDeviceCheck(bool OnOff) override {}
 
-	virtual bool AddDevice( const FString& DeviceName, bool bDefault ) OVERRIDE
+	virtual bool AddDevice( const FString& DeviceName, bool bDefault ) override
 	{
 		return false;
 	}
 
-	virtual void GetAllDevices( TArray<ITargetDevicePtr>& OutDevices ) const OVERRIDE;
+	virtual void GetAllDevices( TArray<ITargetDevicePtr>& OutDevices ) const override;
 
-	virtual ECompressionFlags GetBaseCompressionMethod( ) const OVERRIDE;
+	virtual ECompressionFlags GetBaseCompressionMethod( ) const override;
 
-	virtual bool GenerateStreamingInstallManifest(const TMultiMap<FString, int32>& ChunkMap, const TSet<int32>& ChunkIDsInUse) const OVERRIDE
+	virtual bool GenerateStreamingInstallManifest(const TMultiMap<FString, int32>& ChunkMap, const TSet<int32>& ChunkIDsInUse) const override
 	{
 		return true;
 	}
 
-	virtual ITargetDevicePtr GetDefaultDevice( ) const OVERRIDE;
+	virtual ITargetDevicePtr GetDefaultDevice( ) const override;
 
-	virtual ITargetDevicePtr GetDevice( const FTargetDeviceId& DeviceId ) OVERRIDE;
+	virtual ITargetDevicePtr GetDevice( const FTargetDeviceId& DeviceId ) override;
 
-	virtual bool IsRunningPlatform( ) const OVERRIDE;
+	virtual bool IsRunningPlatform( ) const override;
 
-	virtual bool IsServerOnly( ) const OVERRIDE
+	virtual bool IsServerOnly( ) const override
 	{
 		return false;
 	}
 
-	virtual bool IsSdkInstalled(bool bProjectHasCode, FString& OutDocumentationPath) const OVERRIDE;
+	virtual bool IsSdkInstalled(bool bProjectHasCode, FString& OutDocumentationPath) const override;
 
-	virtual bool SupportsFeature( ETargetPlatformFeatures::Type Feature ) const OVERRIDE;
+	virtual bool SupportsFeature( ETargetPlatformFeatures::Type Feature ) const override;
 
 	virtual bool SupportsTextureFormat( FName Format ) const 
 	{
@@ -114,32 +114,32 @@ public:
 	}
 
 #if WITH_ENGINE
-	virtual void GetReflectionCaptureFormats( TArray<FName>& OutFormats ) const OVERRIDE
+	virtual void GetReflectionCaptureFormats( TArray<FName>& OutFormats ) const override
 	{
 		OutFormats.Add(FName(TEXT("EncodedHDR")));
 	}
 
-	virtual void GetAllPossibleShaderFormats( TArray<FName>& OutFormats ) const OVERRIDE;
+	virtual void GetAllPossibleShaderFormats( TArray<FName>& OutFormats ) const override;
 
-	virtual void GetAllTargetedShaderFormats(TArray<FName>& OutFormats) const OVERRIDE;
+	virtual void GetAllTargetedShaderFormats(TArray<FName>& OutFormats) const override;
 
-	virtual const class FStaticMeshLODSettings& GetStaticMeshLODSettings() const OVERRIDE;
+	virtual const class FStaticMeshLODSettings& GetStaticMeshLODSettings() const override;
 
-	virtual void GetTextureFormats( const UTexture* InTexture, TArray<FName>& OutFormats ) const OVERRIDE;
+	virtual void GetTextureFormats( const UTexture* InTexture, TArray<FName>& OutFormats ) const override;
 
-	virtual const struct FTextureLODSettings& GetTextureLODSettings( ) const OVERRIDE;
+	virtual const struct FTextureLODSettings& GetTextureLODSettings( ) const override;
 
-	virtual FName GetWaveFormat( class USoundWave* Wave ) const OVERRIDE;
+	virtual FName GetWaveFormat( class USoundWave* Wave ) const override;
 #endif //WITH_ENGINE
 
 	DECLARE_DERIVED_EVENT(FAndroidTargetPlatform, ITargetPlatform::FOnTargetDeviceDiscovered, FOnTargetDeviceDiscovered);
-	virtual FOnTargetDeviceDiscovered& OnDeviceDiscovered( ) OVERRIDE
+	virtual FOnTargetDeviceDiscovered& OnDeviceDiscovered( ) override
 	{
 		return DeviceDiscoveredEvent;
 	}
 
 	DECLARE_DERIVED_EVENT(FAndroidTargetPlatform, ITargetPlatform::FOnTargetDeviceLost, FOnTargetDeviceLost);
-	virtual FOnTargetDeviceLost& OnDeviceLost( ) OVERRIDE
+	virtual FOnTargetDeviceLost& OnDeviceLost( ) override
 	{
 		return DeviceLostEvent;
 	}
diff --git a/Engine/Source/Developer/Android/AndroidTargetPlatform/Private/AndroidTargetPlatformModule.cpp b/Engine/Source/Developer/Android/AndroidTargetPlatform/Private/AndroidTargetPlatformModule.cpp
--- a/Engine/Source/Developer/Android/AndroidTargetPlatform/Private/AndroidTargetPlatformModule.cpp
+++ b/Engine/Source/Developer/Android/AndroidTargetPlatform/Private/AndroidTargetPlatformModule.cpp
@@ -29,7 +29,7 @@ public:
 	
 	// Begin ITargetPlatformModule interface
 
-	virtual ITargetPlatform* GetTargetPlatform() OVERRIDE
+	virtual ITargetPlatform* GetTargetPlatform() override
 	{
 		if (AndroidTargetSingleton == NULL)
 		{
@@ -43,11 +43,11 @@ public:
 
 public:
 	// Begin IModuleInterface interface
-	virtual void StartupModule() OVERRIDE
+	virtual void StartupModule() override
 	{
 	}
 
-	virtual void ShutdownModule() OVERRIDE
+	virtual void ShutdownModule() override
 	{
 	}
 	// End IModuleInterface interface
diff --git a/Engine/Source/Developer/Android/Android_ATCTargetPlatform/Private/Android_ATCTargetPlatformModule.cpp b/Engine/Source/Developer/Android/Android_ATCTargetPlatform/Private/Android_ATCTargetPlatformModule.cpp
--- a/Engine/Source/Developer/Android/Android_ATCTargetPlatform/Private/Android_ATCTargetPlatformModule.cpp
+++ b/Engine/Source/Developer/Android/Android_ATCTargetPlatform/Private/Android_ATCTargetPlatformModule.cpp
@@ -15,22 +15,22 @@
 class FAndroid_ATCTargetPlatform
 	: public FAndroidTargetPlatform
 {
-	virtual FString GetAndroidVariantName( ) OVERRIDE
+	virtual FString GetAndroidVariantName( ) override
 	{
 		return TEXT("ATC");
 	}
 
-	virtual FText DisplayName( ) const OVERRIDE
+	virtual FText DisplayName( ) const override
 	{
 		return LOCTEXT("Android_ATC", "Android (ATC)");
 	}
 
-	virtual FString PlatformName() const OVERRIDE
+	virtual FString PlatformName() const override
 	{
 		return FString(FAndroid_ATCCPlatformProperties::PlatformName());
 	}
 
-	virtual bool SupportsTextureFormat( FName Format ) const OVERRIDE
+	virtual bool SupportsTextureFormat( FName Format ) const override
 	{
 		if( Format == AndroidTexFormat::NameATC_RGB ||
 			Format == AndroidTexFormat::NameATC_RGBA_I ||
@@ -41,7 +41,7 @@ class FAndroid_ATCTargetPlatform
 		return false;
 	}
 
-	virtual bool SupportedByExtensionsString( const FString& ExtensionsString, const int GLESVersion ) const OVERRIDE
+	virtual bool SupportedByExtensionsString( const FString& ExtensionsString, const int GLESVersion ) const override
 	{
 		return (ExtensionsString.Contains(TEXT("GL_ATI_texture_compression_atitc")) || ExtensionsString.Contains(TEXT("GL_AMD_compressed_ATC_texture")));
 	}
@@ -75,7 +75,7 @@ public:
 	
 	// Begin ITargetPlatformModule interface
 
-	virtual ITargetPlatform* GetTargetPlatform() OVERRIDE
+	virtual ITargetPlatform* GetTargetPlatform() override
 	{
 		if (AndroidTargetSingleton == NULL)
 		{
diff --git a/Engine/Source/Developer/Android/Android_DXTTargetPlatform/Private/Android_DXTTargetPlatformModule.cpp b/Engine/Source/Developer/Android/Android_DXTTargetPlatform/Private/Android_DXTTargetPlatformModule.cpp
--- a/Engine/Source/Developer/Android/Android_DXTTargetPlatform/Private/Android_DXTTargetPlatformModule.cpp
+++ b/Engine/Source/Developer/Android/Android_DXTTargetPlatform/Private/Android_DXTTargetPlatformModule.cpp
@@ -15,22 +15,22 @@
 class FAndroid_DXTTargetPlatform
 	: public FAndroidTargetPlatform
 {
-	virtual FString GetAndroidVariantName( ) OVERRIDE
+	virtual FString GetAndroidVariantName( ) override
 	{
 		return TEXT("DXT");
 	}
 
-	virtual FText DisplayName( ) const OVERRIDE
+	virtual FText DisplayName( ) const override
 	{
 		return LOCTEXT("Android_DXT", "Android (DXT)");
 	}
 
-	virtual FString PlatformName() const OVERRIDE
+	virtual FString PlatformName() const override
 	{
 		return FString(FAndroid_DXTPlatformProperties::PlatformName());
 	}
 
-	virtual bool SupportsTextureFormat( FName Format ) const OVERRIDE
+	virtual bool SupportsTextureFormat( FName Format ) const override
 	{
 		if( Format == AndroidTexFormat::NameDXT1 ||
 			Format == AndroidTexFormat::NameDXT5 ||
@@ -41,7 +41,7 @@ class FAndroid_DXTTargetPlatform
 		return false;
 	}
 
-	virtual bool SupportedByExtensionsString( const FString& ExtensionsString, const int GLESVersion ) const OVERRIDE
+	virtual bool SupportedByExtensionsString( const FString& ExtensionsString, const int GLESVersion ) const override
 	{
 		return (ExtensionsString.Contains(TEXT("GL_NV_texture_compression_s3tc")) || ExtensionsString.Contains(TEXT("GL_EXT_texture_compression_s3tc")));
 	}
@@ -74,7 +74,7 @@ public:
 	
 	// Begin ITargetPlatformModule interface
 
-	virtual ITargetPlatform* GetTargetPlatform() OVERRIDE
+	virtual ITargetPlatform* GetTargetPlatform() override
 	{
 		if (AndroidTargetSingleton == NULL)
 		{
diff --git a/Engine/Source/Developer/Android/Android_ETC1TargetPlatform/Private/Android_ETC1TargetPlatformModule.cpp b/Engine/Source/Developer/Android/Android_ETC1TargetPlatform/Private/Android_ETC1TargetPlatformModule.cpp
--- a/Engine/Source/Developer/Android/Android_ETC1TargetPlatform/Private/Android_ETC1TargetPlatformModule.cpp
+++ b/Engine/Source/Developer/Android/Android_ETC1TargetPlatform/Private/Android_ETC1TargetPlatformModule.cpp
@@ -19,7 +19,7 @@ public:
 
 	// Begin FAndroidTargetPlatform overrides
 
-	virtual FText DisplayName( ) const OVERRIDE
+	virtual FText DisplayName( ) const override
 	{
 		return LOCTEXT("Android_ETC1", "Android (ETC1)");
 	}
@@ -29,12 +29,12 @@ public:
 		return TEXT("ETC1");
 	}
 
-	virtual FString PlatformName() const OVERRIDE
+	virtual FString PlatformName() const override
 	{
 		return FString(FAndroid_ETC1PlatformProperties::PlatformName());
 	}
 
-	virtual bool SupportsTextureFormat( FName Format ) const OVERRIDE
+	virtual bool SupportsTextureFormat( FName Format ) const override
 	{
 		if (Format == AndroidTexFormat::NameETC1 || 
 			Format == AndroidTexFormat::NameAutoETC1)
@@ -47,7 +47,7 @@ public:
 
 	// End FAndroidTargetPlatform overrides
 
-	virtual bool SupportedByExtensionsString( const FString& ExtensionsString, const int GLESVersion ) const OVERRIDE
+	virtual bool SupportedByExtensionsString( const FString& ExtensionsString, const int GLESVersion ) const override
 	{
 		return ExtensionsString.Contains(TEXT("GL_OES_compressed_ETC1_RGB8_texture"));
 	}
@@ -80,7 +80,7 @@ public:
 	
 	// Begin ITargetPlatformModule interface
 
-	virtual ITargetPlatform* GetTargetPlatform() OVERRIDE
+	virtual ITargetPlatform* GetTargetPlatform() override
 	{
 		if (AndroidTargetSingleton == NULL)
 		{
diff --git a/Engine/Source/Developer/Android/Android_ETC2TargetPlatform/Private/Android_ETC2TargetPlatformModule.cpp b/Engine/Source/Developer/Android/Android_ETC2TargetPlatform/Private/Android_ETC2TargetPlatformModule.cpp
--- a/Engine/Source/Developer/Android/Android_ETC2TargetPlatform/Private/Android_ETC2TargetPlatformModule.cpp
+++ b/Engine/Source/Developer/Android/Android_ETC2TargetPlatform/Private/Android_ETC2TargetPlatformModule.cpp
@@ -19,7 +19,7 @@ public:
 
 	// Begin FAndroidTargetPlatform overrides
 
-	virtual FText DisplayName( ) const OVERRIDE
+	virtual FText DisplayName( ) const override
 	{
 		return LOCTEXT("Android_ETC2", "Android (ETC2)");
 	}
@@ -29,12 +29,12 @@ public:
 		return TEXT("ETC2");
 	}
 
-	virtual FString PlatformName() const OVERRIDE
+	virtual FString PlatformName() const override
 	{
 		return FString(FAndroid_ETC2PlatformProperties::PlatformName());
 	}
 
-	virtual bool SupportsTextureFormat( FName Format ) const OVERRIDE
+	virtual bool SupportsTextureFormat( FName Format ) const override
 	{
 		if (Format == AndroidTexFormat::NameETC2_RGB ||
 			Format == AndroidTexFormat::NameETC2_RGBA ||
@@ -48,7 +48,7 @@ public:
 
 	// End FAndroidTargetPlatform overrides
 
-	virtual bool SupportedByExtensionsString( const FString& ExtensionsString, const int GLESVersion ) const OVERRIDE
+	virtual bool SupportedByExtensionsString( const FString& ExtensionsString, const int GLESVersion ) const override
 	{
 		return GLESVersion >= 0x30000;
 	}
@@ -81,7 +81,7 @@ public:
 	
 	// Begin ITargetPlatformModule interface
 
-	virtual ITargetPlatform* GetTargetPlatform() OVERRIDE
+	virtual ITargetPlatform* GetTargetPlatform() override
 	{
 		if (AndroidTargetSingleton == NULL)
 		{
diff --git a/Engine/Source/Developer/AssetTools/Private/AssetTools.h b/Engine/Source/Developer/AssetTools/Private/AssetTools.h
--- a/Engine/Source/Developer/AssetTools/Private/AssetTools.h
+++ b/Engine/Source/Developer/AssetTools/Private/AssetTools.h
@@ -9,24 +9,24 @@ public:
 	virtual ~FAssetTools();
 
 	// IAssetTools implementation
-	virtual void RegisterAssetTypeActions(const TSharedRef<IAssetTypeActions>& NewActions) OVERRIDE;
-	virtual void UnregisterAssetTypeActions(const TSharedRef<IAssetTypeActions>& ActionsToRemove) OVERRIDE;
-	virtual void GetAssetTypeActionsList( TArray<TWeakPtr<IAssetTypeActions>>& OutAssetTypeActionsList ) const OVERRIDE;
-	virtual TWeakPtr<IAssetTypeActions> GetAssetTypeActionsForClass( UClass* Class ) const OVERRIDE;
-	virtual bool GetAssetActions( const TArray<UObject*>& InObjects, FMenuBuilder& MenuBuilder, bool bIncludeHeading = true ) OVERRIDE;
-	virtual UObject* CreateAsset(const FString& AssetName, const FString& PackagePath, UClass* AssetClass, UFactory* Factory, FName CallingContext = NAME_None) OVERRIDE;
-	virtual UObject* DuplicateAsset(const FString& AssetName, const FString& PackagePath, UObject* OriginalObject) OVERRIDE;
-	virtual void RenameAssets(const TArray<FAssetRenameData>& AssetsAndNames) const OVERRIDE;
-	virtual TArray<UObject*> ImportAssets(const FString& DestinationPath) OVERRIDE;
-	virtual TArray<UObject*> ImportAssets(const TArray<FString>& Files, const FString& DestinationPath) const OVERRIDE;
-	virtual void CreateUniqueAssetName(const FString& InBasePackageName, const FString& InSuffix, FString& OutPackageName, FString& OutAssetName) const OVERRIDE;
-	virtual bool AssetUsesGenericThumbnail( const FAssetData& AssetData ) const OVERRIDE;
-	virtual void DiffAgainstDepot(UObject* InObject, const FString& InPackagePath, const FString& InPackageName) const OVERRIDE;
-	virtual void DiffAssets(UObject* OldAsset1, UObject* NewAsset, const struct FRevisionInfo& OldRevision, const struct FRevisionInfo& NewRevision) const OVERRIDE;
-	virtual FString DumpAssetToTempFile(UObject* Asset) const OVERRIDE;
-	virtual bool CreateDiffProcess(const FString& DiffCommand, const FString& OldTextFilename, const FString& NewTextFilename, const FString& DiffArgs = FString("")) const OVERRIDE;
-	virtual void MigratePackages(const TArray<FName>& PackageNamesToMigrate) const OVERRIDE;
-	virtual void FixupReferencers(const TArray<UObjectRedirector*>& Objects) const OVERRIDE;
+	virtual void RegisterAssetTypeActions(const TSharedRef<IAssetTypeActions>& NewActions) override;
+	virtual void UnregisterAssetTypeActions(const TSharedRef<IAssetTypeActions>& ActionsToRemove) override;
+	virtual void GetAssetTypeActionsList( TArray<TWeakPtr<IAssetTypeActions>>& OutAssetTypeActionsList ) const override;
+	virtual TWeakPtr<IAssetTypeActions> GetAssetTypeActionsForClass( UClass* Class ) const override;
+	virtual bool GetAssetActions( const TArray<UObject*>& InObjects, FMenuBuilder& MenuBuilder, bool bIncludeHeading = true ) override;
+	virtual UObject* CreateAsset(const FString& AssetName, const FString& PackagePath, UClass* AssetClass, UFactory* Factory, FName CallingContext = NAME_None) override;
+	virtual UObject* DuplicateAsset(const FString& AssetName, const FString& PackagePath, UObject* OriginalObject) override;
+	virtual void RenameAssets(const TArray<FAssetRenameData>& AssetsAndNames) const override;
+	virtual TArray<UObject*> ImportAssets(const FString& DestinationPath) override;
+	virtual TArray<UObject*> ImportAssets(const TArray<FString>& Files, const FString& DestinationPath) const override;
+	virtual void CreateUniqueAssetName(const FString& InBasePackageName, const FString& InSuffix, FString& OutPackageName, FString& OutAssetName) const override;
+	virtual bool AssetUsesGenericThumbnail( const FAssetData& AssetData ) const override;
+	virtual void DiffAgainstDepot(UObject* InObject, const FString& InPackagePath, const FString& InPackageName) const override;
+	virtual void DiffAssets(UObject* OldAsset1, UObject* NewAsset, const struct FRevisionInfo& OldRevision, const struct FRevisionInfo& NewRevision) const override;
+	virtual FString DumpAssetToTempFile(UObject* Asset) const override;
+	virtual bool CreateDiffProcess(const FString& DiffCommand, const FString& OldTextFilename, const FString& NewTextFilename, const FString& DiffArgs = FString("")) const override;
+	virtual void MigratePackages(const TArray<FName>& PackageNamesToMigrate) const override;
+	virtual void FixupReferencers(const TArray<UObjectRedirector*>& Objects) const override;
 
 public:
 	/** Gets the asset tools singleton as a FAssetTools for asset tools module use */
diff --git a/Engine/Source/Developer/AssetTools/Private/AssetTypeActions/AssetTypeActions_AimOffset.h b/Engine/Source/Developer/AssetTools/Private/AssetTypeActions/AssetTypeActions_AimOffset.h
--- a/Engine/Source/Developer/AssetTools/Private/AssetTypeActions/AssetTypeActions_AimOffset.h
+++ b/Engine/Source/Developer/AssetTools/Private/AssetTypeActions/AssetTypeActions_AimOffset.h
@@ -6,7 +6,7 @@ class FAssetTypeActions_AimOffset : public FAssetTypeActions_BlendSpace
 {
 public:
 	// IAssetTypeActions Implementation
-	virtual FText GetName() const OVERRIDE { return NSLOCTEXT("AssetTypeActions", "AssetTypeActions_AimOffset", "Aim Offset"); }
-	virtual FColor GetTypeColor() const OVERRIDE { return FColor(0,162,232); }
-	virtual UClass* GetSupportedClass() const OVERRIDE { return UAimOffsetBlendSpace::StaticClass(); }
+	virtual FText GetName() const override { return NSLOCTEXT("AssetTypeActions", "AssetTypeActions_AimOffset", "Aim Offset"); }
+	virtual FColor GetTypeColor() const override { return FColor(0,162,232); }
+	virtual UClass* GetSupportedClass() const override { return UAimOffsetBlendSpace::StaticClass(); }
 };
\ No newline at end of file
diff --git a/Engine/Source/Developer/AssetTools/Private/AssetTypeActions/AssetTypeActions_AimOffset1D.h b/Engine/Source/Developer/AssetTools/Private/AssetTypeActions/AssetTypeActions_AimOffset1D.h
--- a/Engine/Source/Developer/AssetTools/Private/AssetTypeActions/AssetTypeActions_AimOffset1D.h
+++ b/Engine/Source/Developer/AssetTools/Private/AssetTypeActions/AssetTypeActions_AimOffset1D.h
@@ -6,7 +6,7 @@ class FAssetTypeActions_AimOffset1D : public FAssetTypeActions_BlendSpace1D
 {
 public:
 	// IAssetTypeActions Implementation
-	virtual FText GetName() const OVERRIDE { return NSLOCTEXT("AssetTypeActions", "AssetTypeActions_AimOffset1D", "Aim Offset 1D"); }
-	virtual FColor GetTypeColor() const OVERRIDE { return FColor(0,162,232); }
-	virtual UClass* GetSupportedClass() const OVERRIDE { return UAimOffsetBlendSpace1D::StaticClass(); }
+	virtual FText GetName() const override { return NSLOCTEXT("AssetTypeActions", "AssetTypeActions_AimOffset1D", "Aim Offset 1D"); }
+	virtual FColor GetTypeColor() const override { return FColor(0,162,232); }
+	virtual UClass* GetSupportedClass() const override { return UAimOffsetBlendSpace1D::StaticClass(); }
 };
\ No newline at end of file
diff --git a/Engine/Source/Developer/AssetTools/Private/AssetTypeActions/AssetTypeActions_AnimBlueprint.h b/Engine/Source/Developer/AssetTools/Private/AssetTypeActions/AssetTypeActions_AnimBlueprint.h
--- a/Engine/Source/Developer/AssetTools/Private/AssetTypeActions/AssetTypeActions_AnimBlueprint.h
+++ b/Engine/Source/Developer/AssetTools/Private/AssetTypeActions/AssetTypeActions_AnimBlueprint.h
@@ -6,12 +6,12 @@ class FAssetTypeActions_AnimBlueprint : public FAssetTypeActions_Blueprint
 {
 public:
 	// IAssetTypeActions Implementation
-	virtual FText GetName() const OVERRIDE { return NSLOCTEXT("AssetTypeActions", "AssetTypeActions_AnimBlueprint", "Animation Blueprint"); }
-	virtual FColor GetTypeColor() const OVERRIDE { return FColor(200,116,0); }
-	virtual UClass* GetSupportedClass() const OVERRIDE { return UAnimBlueprint::StaticClass(); }
-	virtual void GetActions( const TArray<UObject*>& InObjects, FMenuBuilder& MenuBuilder ) OVERRIDE;
-	virtual void OpenAssetEditor( const TArray<UObject*>& InObjects, TSharedPtr<class IToolkitHost> EditWithinLevelEditor = TSharedPtr<IToolkitHost>() ) OVERRIDE;
-	virtual uint32 GetCategories() OVERRIDE { return EAssetTypeCategories::Animation; }
+	virtual FText GetName() const override { return NSLOCTEXT("AssetTypeActions", "AssetTypeActions_AnimBlueprint", "Animation Blueprint"); }
+	virtual FColor GetTypeColor() const override { return FColor(200,116,0); }
+	virtual UClass* GetSupportedClass() const override { return UAnimBlueprint::StaticClass(); }
+	virtual void GetActions( const TArray<UObject*>& InObjects, FMenuBuilder& MenuBuilder ) override;
+	virtual void OpenAssetEditor( const TArray<UObject*>& InObjects, TSharedPtr<class IToolkitHost> EditWithinLevelEditor = TSharedPtr<IToolkitHost>() ) override;
+	virtual uint32 GetCategories() override { return EAssetTypeCategories::Animation; }
 
 private:
 
diff --git a/Engine/Source/Developer/AssetTools/Private/AssetTypeActions/AssetTypeActions_AnimComposite.h b/Engine/Source/Developer/AssetTools/Private/AssetTypeActions/AssetTypeActions_AnimComposite.h
--- a/Engine/Source/Developer/AssetTools/Private/AssetTypeActions/AssetTypeActions_AnimComposite.h
+++ b/Engine/Source/Developer/AssetTools/Private/AssetTypeActions/AssetTypeActions_AnimComposite.h
@@ -6,8 +6,8 @@ class FAssetTypeActions_AnimComposite : public FAssetTypeActions_AnimationAsset
 {
 public:
 	// IAssetTypeActions Implementation
-	virtual FText GetName() const OVERRIDE { return NSLOCTEXT("AssetTypeActions", "AssetTypeActions_AnimComposite", "Animation Composite"); }
-	virtual FColor GetTypeColor() const OVERRIDE { return FColor(0,128,192); }
-	virtual UClass* GetSupportedClass() const OVERRIDE { return UAnimComposite::StaticClass(); }
-	virtual bool CanFilter() OVERRIDE { return true; }
+	virtual FText GetName() const override { return NSLOCTEXT("AssetTypeActions", "AssetTypeActions_AnimComposite", "Animation Composite"); }
+	virtual FColor GetTypeColor() const override { return FColor(0,128,192); }
+	virtual UClass* GetSupportedClass() const override { return UAnimComposite::StaticClass(); }
+	virtual bool CanFilter() override { return true; }
 };
\ No newline at end of file
diff --git a/Engine/Source/Developer/AssetTools/Private/AssetTypeActions/AssetTypeActions_AnimMontage.h b/Engine/Source/Developer/AssetTools/Private/AssetTypeActions/AssetTypeActions_AnimMontage.h
--- a/Engine/Source/Developer/AssetTools/Private/AssetTypeActions/AssetTypeActions_AnimMontage.h
+++ b/Engine/Source/Developer/AssetTools/Private/AssetTypeActions/AssetTypeActions_AnimMontage.h
@@ -6,8 +6,8 @@ class FAssetTypeActions_AnimMontage : public FAssetTypeActions_AnimationAsset
 {
 public:
 	// IAssetTypeActions Implementation
-	virtual FText GetName() const OVERRIDE { return NSLOCTEXT("AssetTypeActions", "AssetTypeActions_AnimMontage", "Animation Montage"); }
-	virtual FColor GetTypeColor() const OVERRIDE { return FColor(100,100,255); }
-	virtual UClass* GetSupportedClass() const OVERRIDE { return UAnimMontage::StaticClass(); }
-	virtual bool CanFilter() OVERRIDE { return true; }
+	virtual FText GetName() const override { return NSLOCTEXT("AssetTypeActions", "AssetTypeActions_AnimMontage", "Animation Montage"); }
+	virtual FColor GetTypeColor() const override { return FColor(100,100,255); }
+	virtual UClass* GetSupportedClass() const override { return UAnimMontage::StaticClass(); }
+	virtual bool CanFilter() override { return true; }
 };
\ No newline at end of file
diff --git a/Engine/Source/Developer/AssetTools/Private/AssetTypeActions/AssetTypeActions_AnimSequence.h b/Engine/Source/Developer/AssetTools/Private/AssetTypeActions/AssetTypeActions_AnimSequence.h
--- a/Engine/Source/Developer/AssetTools/Private/AssetTypeActions/AssetTypeActions_AnimSequence.h
+++ b/Engine/Source/Developer/AssetTools/Private/AssetTypeActions/AssetTypeActions_AnimSequence.h
@@ -9,10 +9,10 @@ class FAssetTypeActions_AnimSequence : public FAssetTypeActions_AnimationAsset
 {
 public:
 	// IAssetTypeActions Implementation
-	virtual FText GetName() const OVERRIDE { return NSLOCTEXT("AssetTypeActions", "AssetTypeActions_AnimSequence", "Animation Sequence"); }
-	virtual UClass* GetSupportedClass() const OVERRIDE { return UAnimSequence::StaticClass(); }
-	virtual void GetActions( const TArray<UObject*>& InObjects, FMenuBuilder& MenuBuilder ) OVERRIDE;
-	virtual bool CanFilter() OVERRIDE { return true; }
+	virtual FText GetName() const override { return NSLOCTEXT("AssetTypeActions", "AssetTypeActions_AnimSequence", "Animation Sequence"); }
+	virtual UClass* GetSupportedClass() const override { return UAnimSequence::StaticClass(); }
+	virtual void GetActions( const TArray<UObject*>& InObjects, FMenuBuilder& MenuBuilder ) override;
+	virtual bool CanFilter() override { return true; }
 
 private:
 	/** Handler for when Reimport is selected */
diff --git a/Engine/Source/Developer/AssetTools/Private/AssetTypeActions/AssetTypeActions_AnimationAsset.h b/Engine/Source/Developer/AssetTools/Private/AssetTypeActions/AssetTypeActions_AnimationAsset.h
--- a/Engine/Source/Developer/AssetTools/Private/AssetTypeActions/AssetTypeActions_AnimationAsset.h
+++ b/Engine/Source/Developer/AssetTools/Private/AssetTypeActions/AssetTypeActions_AnimationAsset.h
@@ -6,14 +6,14 @@ class FAssetTypeActions_AnimationAsset : public FAssetTypeActions_Base
 {
 public:
 	// IAssetTypeActions Implementation
-	virtual FText GetName() const OVERRIDE { return NSLOCTEXT("AssetTypeActions", "AssetTypeActions_AnimationAsset", "AnimationAsset"); }
-	virtual FColor GetTypeColor() const OVERRIDE { return FColor(80,123,72); }
-	virtual UClass* GetSupportedClass() const OVERRIDE { return UAnimationAsset::StaticClass(); }
-	virtual bool HasActions ( const TArray<UObject*>& InObjects ) const OVERRIDE { return true; }
-	virtual void GetActions( const TArray<UObject*>& InObjects, FMenuBuilder& MenuBuilder ) OVERRIDE;
-	virtual void OpenAssetEditor( const TArray<UObject*>& InObjects, TSharedPtr<class IToolkitHost> EditWithinLevelEditor = TSharedPtr<IToolkitHost>() ) OVERRIDE;
-	virtual bool CanFilter() OVERRIDE { return false; }
-	virtual uint32 GetCategories() OVERRIDE { return EAssetTypeCategories::Animation; }
+	virtual FText GetName() const override { return NSLOCTEXT("AssetTypeActions", "AssetTypeActions_AnimationAsset", "AnimationAsset"); }
+	virtual FColor GetTypeColor() const override { return FColor(80,123,72); }
+	virtual UClass* GetSupportedClass() const override { return UAnimationAsset::StaticClass(); }
+	virtual bool HasActions ( const TArray<UObject*>& InObjects ) const override { return true; }
+	virtual void GetActions( const TArray<UObject*>& InObjects, FMenuBuilder& MenuBuilder ) override;
+	virtual void OpenAssetEditor( const TArray<UObject*>& InObjects, TSharedPtr<class IToolkitHost> EditWithinLevelEditor = TSharedPtr<IToolkitHost>() ) override;
+	virtual bool CanFilter() override { return false; }
+	virtual uint32 GetCategories() override { return EAssetTypeCategories::Animation; }
 
 private:
 	/** Handler to fill the retarget submenu */
diff --git a/Engine/Source/Developer/AssetTools/Private/AssetTypeActions/AssetTypeActions_BlendSpace.h b/Engine/Source/Developer/AssetTools/Private/AssetTypeActions/AssetTypeActions_BlendSpace.h
--- a/Engine/Source/Developer/AssetTools/Private/AssetTypeActions/AssetTypeActions_BlendSpace.h
+++ b/Engine/Source/Developer/AssetTools/Private/AssetTypeActions/AssetTypeActions_BlendSpace.h
@@ -6,8 +6,8 @@ class FAssetTypeActions_BlendSpace : public FAssetTypeActions_AnimationAsset
 {
 public:
 	// IAssetTypeActions Implementation
-	virtual FText GetName() const OVERRIDE { return NSLOCTEXT("AssetTypeActions", "AssetTypeActions_BlendSpace", "Blend Space"); }
-	virtual FColor GetTypeColor() const OVERRIDE { return FColor(255,168,111); }
-	virtual UClass* GetSupportedClass() const OVERRIDE { return UBlendSpace::StaticClass(); }
-	virtual bool CanFilter() OVERRIDE { return true; }
+	virtual FText GetName() const override { return NSLOCTEXT("AssetTypeActions", "AssetTypeActions_BlendSpace", "Blend Space"); }
+	virtual FColor GetTypeColor() const override { return FColor(255,168,111); }
+	virtual UClass* GetSupportedClass() const override { return UBlendSpace::StaticClass(); }
+	virtual bool CanFilter() override { return true; }
 };
\ No newline at end of file
diff --git a/Engine/Source/Developer/AssetTools/Private/AssetTypeActions/AssetTypeActions_BlendSpace1D.h b/Engine/Source/Developer/AssetTools/Private/AssetTypeActions/AssetTypeActions_BlendSpace1D.h
--- a/Engine/Source/Developer/AssetTools/Private/AssetTypeActions/AssetTypeActions_BlendSpace1D.h
+++ b/Engine/Source/Developer/AssetTools/Private/AssetTypeActions/AssetTypeActions_BlendSpace1D.h
@@ -6,8 +6,8 @@ class FAssetTypeActions_BlendSpace1D : public FAssetTypeActions_AnimationAsset
 {
 public:
 	// IAssetTypeActions Implementation
-	virtual FText GetName() const OVERRIDE { return NSLOCTEXT("AssetTypeActions", "AssetTypeActions_BlendSpace1D", "Blend Space 1D"); }
-	virtual FColor GetTypeColor() const OVERRIDE { return FColor(255,180,130); }
-	virtual UClass* GetSupportedClass() const OVERRIDE { return UBlendSpace1D::StaticClass(); }
-	virtual bool CanFilter() OVERRIDE { return true; }
+	virtual FText GetName() const override { return NSLOCTEXT("AssetTypeActions", "AssetTypeActions_BlendSpace1D", "Blend Space 1D"); }
+	virtual FColor GetTypeColor() const override { return FColor(255,180,130); }
+	virtual UClass* GetSupportedClass() const override { return UBlendSpace1D::StaticClass(); }
+	virtual bool CanFilter() override { return true; }
 };
\ No newline at end of file
diff --git a/Engine/Source/Developer/AssetTools/Private/AssetTypeActions/AssetTypeActions_Blueprint.h b/Engine/Source/Developer/AssetTools/Private/AssetTypeActions/AssetTypeActions_Blueprint.h
--- a/Engine/Source/Developer/AssetTools/Private/AssetTypeActions/AssetTypeActions_Blueprint.h
+++ b/Engine/Source/Developer/AssetTools/Private/AssetTypeActions/AssetTypeActions_Blueprint.h
@@ -6,16 +6,16 @@ class FAssetTypeActions_Blueprint : public FAssetTypeActions_Base
 {
 public:
 	// IAssetTypeActions Implementation
-	virtual FText GetName() const OVERRIDE { return NSLOCTEXT("AssetTypeActions", "AssetTypeActions_Blueprint", "Blueprint"); }
-	virtual FColor GetTypeColor() const OVERRIDE { return FColor( 63, 126, 255 ); }
-	virtual UClass* GetSupportedClass() const OVERRIDE { return UBlueprint::StaticClass(); }
-	virtual bool HasActions ( const TArray<UObject*>& InObjects ) const OVERRIDE { return true; }
-	virtual void GetActions( const TArray<UObject*>& InObjects, FMenuBuilder& MenuBuilder ) OVERRIDE;
-	virtual void OpenAssetEditor( const TArray<UObject*>& InObjects, TSharedPtr<class IToolkitHost> EditWithinLevelEditor = TSharedPtr<IToolkitHost>() ) OVERRIDE;
-	virtual uint32 GetCategories() OVERRIDE { return EAssetTypeCategories::Basic; }
-	virtual void PerformAssetDiff(UObject* Asset1, UObject* Asset2, const struct FRevisionInfo& OldRevision, const struct FRevisionInfo& NewRevision) const OVERRIDE;
-	virtual class UThumbnailInfo* GetThumbnailInfo(UObject* Asset) const OVERRIDE;
-	virtual FText GetAssetDescription(const FAssetData& AssetData) const OVERRIDE;
+	virtual FText GetName() const override { return NSLOCTEXT("AssetTypeActions", "AssetTypeActions_Blueprint", "Blueprint"); }
+	virtual FColor GetTypeColor() const override { return FColor( 63, 126, 255 ); }
+	virtual UClass* GetSupportedClass() const override { return UBlueprint::StaticClass(); }
+	virtual bool HasActions ( const TArray<UObject*>& InObjects ) const override { return true; }
+	virtual void GetActions( const TArray<UObject*>& InObjects, FMenuBuilder& MenuBuilder ) override;
+	virtual void OpenAssetEditor( const TArray<UObject*>& InObjects, TSharedPtr<class IToolkitHost> EditWithinLevelEditor = TSharedPtr<IToolkitHost>() ) override;
+	virtual uint32 GetCategories() override { return EAssetTypeCategories::Basic; }
+	virtual void PerformAssetDiff(UObject* Asset1, UObject* Asset2, const struct FRevisionInfo& OldRevision, const struct FRevisionInfo& NewRevision) const override;
+	virtual class UThumbnailInfo* GetThumbnailInfo(UObject* Asset) const override;
+	virtual FText GetAssetDescription(const FAssetData& AssetData) const override;
 
 protected:
 	/** Handler for when Edit is selected */
diff --git a/Engine/Source/Developer/AssetTools/Private/AssetTypeActions/AssetTypeActions_BlueprintGeneratedClass.h b/Engine/Source/Developer/AssetTools/Private/AssetTypeActions/AssetTypeActions_BlueprintGeneratedClass.h
--- a/Engine/Source/Developer/AssetTools/Private/AssetTypeActions/AssetTypeActions_BlueprintGeneratedClass.h
+++ b/Engine/Source/Developer/AssetTools/Private/AssetTypeActions/AssetTypeActions_BlueprintGeneratedClass.h
@@ -6,15 +6,15 @@ class FAssetTypeActions_BlueprintGeneratedClass : public FAssetTypeActions_Base
 {
 public:
 	// IAssetTypeActions Implementation
-	virtual FText GetName() const OVERRIDE { return NSLOCTEXT("AssetTypeActions", "AssetTypeActions_BlueprintGeneratedClass", "Blueprint Generated Class"); }
-	virtual FColor GetTypeColor() const OVERRIDE { return FColor( 40, 78, 165 ); }
-	virtual UClass* GetSupportedClass() const OVERRIDE { return UBlueprintGeneratedClass::StaticClass(); }
-	virtual bool HasActions ( const TArray<UObject*>& InObjects ) const OVERRIDE { return true; }
-	virtual void GetActions( const TArray<UObject*>& InObjects, FMenuBuilder& MenuBuilder ) OVERRIDE;
-	virtual void OpenAssetEditor( const TArray<UObject*>& InObjects, TSharedPtr<class IToolkitHost> EditWithinLevelEditor = TSharedPtr<IToolkitHost>() ) OVERRIDE;
-	virtual uint32 GetCategories() OVERRIDE { return EAssetTypeCategories::Misc; }
-	virtual void PerformAssetDiff(UObject* Asset1, UObject* Asset2, const struct FRevisionInfo& OldRevision, const struct FRevisionInfo& NewRevision) const OVERRIDE;
-	virtual class UThumbnailInfo* GetThumbnailInfo(UObject* Asset) const OVERRIDE;
+	virtual FText GetName() const override { return NSLOCTEXT("AssetTypeActions", "AssetTypeActions_BlueprintGeneratedClass", "Blueprint Generated Class"); }
+	virtual FColor GetTypeColor() const override { return FColor( 40, 78, 165 ); }
+	virtual UClass* GetSupportedClass() const override { return UBlueprintGeneratedClass::StaticClass(); }
+	virtual bool HasActions ( const TArray<UObject*>& InObjects ) const override { return true; }
+	virtual void GetActions( const TArray<UObject*>& InObjects, FMenuBuilder& MenuBuilder ) override;
+	virtual void OpenAssetEditor( const TArray<UObject*>& InObjects, TSharedPtr<class IToolkitHost> EditWithinLevelEditor = TSharedPtr<IToolkitHost>() ) override;
+	virtual uint32 GetCategories() override { return EAssetTypeCategories::Misc; }
+	virtual void PerformAssetDiff(UObject* Asset1, UObject* Asset2, const struct FRevisionInfo& OldRevision, const struct FRevisionInfo& NewRevision) const override;
+	virtual class UThumbnailInfo* GetThumbnailInfo(UObject* Asset) const override;
 
 protected:
 	/** Handler for when Edit is selected */
diff --git a/Engine/Source/Developer/AssetTools/Private/AssetTypeActions/AssetTypeActions_CameraAnim.h b/Engine/Source/Developer/AssetTools/Private/AssetTypeActions/AssetTypeActions_CameraAnim.h
--- a/Engine/Source/Developer/AssetTools/Private/AssetTypeActions/AssetTypeActions_CameraAnim.h
+++ b/Engine/Source/Developer/AssetTools/Private/AssetTypeActions/AssetTypeActions_CameraAnim.h
@@ -6,12 +6,12 @@ class FAssetTypeActions_CameraAnim : public FAssetTypeActions_Base
 {
 public:
 	// IAssetTypeActions Implementation
-	virtual FText GetName() const OVERRIDE { return NSLOCTEXT("AssetTypeActions", "AssetTypeActions_CameraAnim", "Camera Anim"); }
-	virtual FColor GetTypeColor() const OVERRIDE { return FColor(255,0,0); }
-	virtual UClass* GetSupportedClass() const OVERRIDE { return UCameraAnim::StaticClass(); }
-	virtual void OpenAssetEditor( const TArray<UObject*>& InObjects, TSharedPtr<class IToolkitHost> EditWithinLevelEditor = TSharedPtr<IToolkitHost>() ) OVERRIDE;
-	virtual bool CanFilter() OVERRIDE { return true; }
-	virtual uint32 GetCategories() OVERRIDE { return EAssetTypeCategories::Misc; }
+	virtual FText GetName() const override { return NSLOCTEXT("AssetTypeActions", "AssetTypeActions_CameraAnim", "Camera Anim"); }
+	virtual FColor GetTypeColor() const override { return FColor(255,0,0); }
+	virtual UClass* GetSupportedClass() const override { return UCameraAnim::StaticClass(); }
+	virtual void OpenAssetEditor( const TArray<UObject*>& InObjects, TSharedPtr<class IToolkitHost> EditWithinLevelEditor = TSharedPtr<IToolkitHost>() ) override;
+	virtual bool CanFilter() override { return true; }
+	virtual uint32 GetCategories() override { return EAssetTypeCategories::Misc; }
 
 private:
 	/** Helper function used to create a preview matinee actor for use when editing a camera anim in Matinee */
diff --git a/Engine/Source/Developer/AssetTools/Private/AssetTypeActions/AssetTypeActions_Curve.h b/Engine/Source/Developer/AssetTools/Private/AssetTypeActions/AssetTypeActions_Curve.h
--- a/Engine/Source/Developer/AssetTools/Private/AssetTypeActions/AssetTypeActions_Curve.h
+++ b/Engine/Source/Developer/AssetTools/Private/AssetTypeActions/AssetTypeActions_Curve.h
@@ -6,14 +6,14 @@ class FAssetTypeActions_Curve : public FAssetTypeActions_Base
 {
 public:
 	// IAssetTypeActions Implementation
-	virtual FText GetName() const OVERRIDE { return NSLOCTEXT("AssetTypeActions", "AssetTypeActions_Curve", "Curve"); }
-	virtual FColor GetTypeColor() const OVERRIDE { return FColor(78, 40, 165); }
-	virtual UClass* GetSupportedClass() const OVERRIDE { return UCurveBase::StaticClass(); }
-	virtual bool HasActions ( const TArray<UObject*>& InObjects ) const OVERRIDE { return true; }
-	virtual void GetActions( const TArray<UObject*>& InObjects, FMenuBuilder& MenuBuilder ) OVERRIDE;
-	virtual void OpenAssetEditor( const TArray<UObject*>& InObjects, TSharedPtr<class IToolkitHost> EditWithinLevelEditor = TSharedPtr<IToolkitHost>() ) OVERRIDE;
-	virtual bool CanFilter() OVERRIDE { return false; }
-	virtual uint32 GetCategories() OVERRIDE { return EAssetTypeCategories::Misc; }
+	virtual FText GetName() const override { return NSLOCTEXT("AssetTypeActions", "AssetTypeActions_Curve", "Curve"); }
+	virtual FColor GetTypeColor() const override { return FColor(78, 40, 165); }
+	virtual UClass* GetSupportedClass() const override { return UCurveBase::StaticClass(); }
+	virtual bool HasActions ( const TArray<UObject*>& InObjects ) const override { return true; }
+	virtual void GetActions( const TArray<UObject*>& InObjects, FMenuBuilder& MenuBuilder ) override;
+	virtual void OpenAssetEditor( const TArray<UObject*>& InObjects, TSharedPtr<class IToolkitHost> EditWithinLevelEditor = TSharedPtr<IToolkitHost>() ) override;
+	virtual bool CanFilter() override { return false; }
+	virtual uint32 GetCategories() override { return EAssetTypeCategories::Misc; }
 
 private:
 	/** Handler for when Edit is selected */
diff --git a/Engine/Source/Developer/AssetTools/Private/AssetTypeActions/AssetTypeActions_CurveFloat.h b/Engine/Source/Developer/AssetTools/Private/AssetTypeActions/AssetTypeActions_CurveFloat.h
--- a/Engine/Source/Developer/AssetTools/Private/AssetTypeActions/AssetTypeActions_CurveFloat.h
+++ b/Engine/Source/Developer/AssetTools/Private/AssetTypeActions/AssetTypeActions_CurveFloat.h
@@ -6,7 +6,7 @@ class FAssetTypeActions_CurveFloat : public FAssetTypeActions_Curve
 {
 public:
 	// IAssetTypeActions Implementation
-	virtual FText GetName() const OVERRIDE { return NSLOCTEXT("AssetTypeActions", "AssetTypeActions_CurveFloat", "Float Curve"); }
-	virtual UClass* GetSupportedClass() const OVERRIDE { return UCurveFloat::StaticClass(); }
-	virtual bool CanFilter() OVERRIDE { return true; }
+	virtual FText GetName() const override { return NSLOCTEXT("AssetTypeActions", "AssetTypeActions_CurveFloat", "Float Curve"); }
+	virtual UClass* GetSupportedClass() const override { return UCurveFloat::StaticClass(); }
+	virtual bool CanFilter() override { return true; }
 };
\ No newline at end of file
diff --git a/Engine/Source/Developer/AssetTools/Private/AssetTypeActions/AssetTypeActions_CurveLinearColor.h b/Engine/Source/Developer/AssetTools/Private/AssetTypeActions/AssetTypeActions_CurveLinearColor.h
--- a/Engine/Source/Developer/AssetTools/Private/AssetTypeActions/AssetTypeActions_CurveLinearColor.h
+++ b/Engine/Source/Developer/AssetTools/Private/AssetTypeActions/AssetTypeActions_CurveLinearColor.h
@@ -6,7 +6,7 @@ class FAssetTypeActions_CurveLinearColor : public FAssetTypeActions_Curve
 {
 public:
 	// IAssetTypeActions Implementation
-	virtual FText GetName() const OVERRIDE { return NSLOCTEXT("AssetTypeActions", "AssetTypeActions_CurveLinearColor", "Color Curve"); }
-	virtual UClass* GetSupportedClass() const OVERRIDE { return UCurveLinearColor::StaticClass(); }
-	virtual bool CanFilter() OVERRIDE { return true; }
+	virtual FText GetName() const override { return NSLOCTEXT("AssetTypeActions", "AssetTypeActions_CurveLinearColor", "Color Curve"); }
+	virtual UClass* GetSupportedClass() const override { return UCurveLinearColor::StaticClass(); }
+	virtual bool CanFilter() override { return true; }
 };
\ No newline at end of file
diff --git a/Engine/Source/Developer/AssetTools/Private/AssetTypeActions/AssetTypeActions_CurveTable.h b/Engine/Source/Developer/AssetTools/Private/AssetTypeActions/AssetTypeActions_CurveTable.h
--- a/Engine/Source/Developer/AssetTools/Private/AssetTypeActions/AssetTypeActions_CurveTable.h
+++ b/Engine/Source/Developer/AssetTools/Private/AssetTypeActions/AssetTypeActions_CurveTable.h
@@ -6,11 +6,11 @@ class FAssetTypeActions_CurveTable : public FAssetTypeActions_CSVAssetBase
 {
 public:
 	// IAssetTypeActions Implementation
-	virtual FText GetName() const OVERRIDE { return NSLOCTEXT("AssetTypeActions", "AssetTypeActions_CurveTable", "Curve Table"); }
-	virtual FColor GetTypeColor() const OVERRIDE { return FColor(62, 140, 35); }
-	virtual UClass* GetSupportedClass() const OVERRIDE { return UCurveTable::StaticClass(); }
-	virtual bool HasActions ( const TArray<UObject*>& InObjects ) const OVERRIDE { return true; }
-	virtual void GetActions( const TArray<UObject*>& InObjects, FMenuBuilder& MenuBuilder ) OVERRIDE;
-	virtual void OpenAssetEditor( const TArray<UObject*>& InObjects, TSharedPtr<class IToolkitHost> EditWithinLevelEditor = TSharedPtr<IToolkitHost>() ) OVERRIDE;
-	virtual uint32 GetCategories() OVERRIDE { return EAssetTypeCategories::Misc; }
+	virtual FText GetName() const override { return NSLOCTEXT("AssetTypeActions", "AssetTypeActions_CurveTable", "Curve Table"); }
+	virtual FColor GetTypeColor() const override { return FColor(62, 140, 35); }
+	virtual UClass* GetSupportedClass() const override { return UCurveTable::StaticClass(); }
+	virtual bool HasActions ( const TArray<UObject*>& InObjects ) const override { return true; }
+	virtual void GetActions( const TArray<UObject*>& InObjects, FMenuBuilder& MenuBuilder ) override;
+	virtual void OpenAssetEditor( const TArray<UObject*>& InObjects, TSharedPtr<class IToolkitHost> EditWithinLevelEditor = TSharedPtr<IToolkitHost>() ) override;
+	virtual uint32 GetCategories() override { return EAssetTypeCategories::Misc; }
 };
\ No newline at end of file
