Added basic Analytics support to OculusRift and Morpheus plugins.
#ttp 335858 ANALYTICS:  Would like to add tracking to see how many users are using the various VR plug ins
reviewer Thomas.Sarkanen

diff --git a/Engine/Plugins/OculusRift/Source/OculusRift/Private/OculusRiftHMD.cpp b/Engine/Plugins/OculusRift/Source/OculusRift/Private/OculusRiftHMD.cpp
--- a/Engine/Plugins/OculusRift/Source/OculusRift/Private/OculusRiftHMD.cpp
+++ b/Engine/Plugins/OculusRift/Source/OculusRift/Private/OculusRiftHMD.cpp
@@ -4,6 +4,9 @@
 #include "OculusRiftHMD.h"
 
 #define DEFAULT_PREDICTION_IN_SECONDS 0.035
+#include "EngineAnalytics.h"
+#include "Runtime/Analytics/Analytics/Public/Interfaces/IAnalyticsProvider.h"
+
 
 
 
@@ -954,6 +957,62 @@ void FOculusRiftHMD::OnScreenModeChange(EWindowMode::Type WindowMode)
 	UpdateStereoRenderingParams();
 }
 
+void FOculusRiftHMD::RecordAnalytics()
+{
+	if (FEngineAnalytics::IsAvailable())
+	{
+		// prepare and send analytics data
+		TArray<FAnalyticsEventAttribute> EventAttributes;
+
+		IHeadMountedDisplay::MonitorInfo MonitorInfo;
+		GetHMDMonitorInfo(MonitorInfo);
+		EventAttributes.Add(FAnalyticsEventAttribute(TEXT("DeviceName"), MonitorInfo.MonitorName));
+		EventAttributes.Add(FAnalyticsEventAttribute(TEXT("DisplayId"), MonitorInfo.MonitorId));
+		FString MonResolution(FString::Printf(TEXT("(%d, %d)"), MonitorInfo.ResolutionX, MonitorInfo.ResolutionY));
+		EventAttributes.Add(FAnalyticsEventAttribute(TEXT("Resolution"), MonResolution));
+			
+		EventAttributes.Add(FAnalyticsEventAttribute(TEXT("ChromaAbCorrectionEnabled"), bChromaAbCorrectionEnabled));
+		EventAttributes.Add(FAnalyticsEventAttribute(TEXT("MagEnabled"), bYawDriftCorrectionEnabled));
+		EventAttributes.Add(FAnalyticsEventAttribute(TEXT("DevSettingsEnabled"), bDevSettingsEnabled));
+		EventAttributes.Add(FAnalyticsEventAttribute(TEXT("MotionPredictionEnabled"), (MotionPredictionInSeconds > 0)));
+		EventAttributes.Add(FAnalyticsEventAttribute(TEXT("TiltCorrectionEnabled"), bTiltCorrectionEnabled));
+		EventAttributes.Add(FAnalyticsEventAttribute(TEXT("AccelGain"), AccelGain));
+		const FVector vec = ToFVector_M2U(GetHeadModel());
+		FString formatted(FString::Printf(TEXT("(%f, %f, %f)"), vec[0], vec[1], vec[2]));
+		EventAttributes.Add(FAnalyticsEventAttribute(TEXT("HeadModel"), formatted));
+		EventAttributes.Add(FAnalyticsEventAttribute(TEXT("OverrideInterpupillaryDistance"), bOverrideIPD));
+		if (bOverrideIPD)
+		{
+			EventAttributes.Add(FAnalyticsEventAttribute(TEXT("InterpupillaryDistance"), GetInterpupillaryDistance()));
+		}
+		EventAttributes.Add(FAnalyticsEventAttribute(TEXT("OverrideStereo"), bOverrideStereo));
+		if (bOverrideStereo)
+		{
+			EventAttributes.Add(FAnalyticsEventAttribute(TEXT("HFOV"), HFOVInRadians));
+			EventAttributes.Add(FAnalyticsEventAttribute(TEXT("VFOV"), VFOVInRadians));
+		}
+		EventAttributes.Add(FAnalyticsEventAttribute(TEXT("OverrideVSync"), bOverrideVSync));
+		if (bOverrideVSync)
+		{
+			EventAttributes.Add(FAnalyticsEventAttribute(TEXT("VSync"), bVSync));
+		}
+		EventAttributes.Add(FAnalyticsEventAttribute(TEXT("OverrideScreenPercentage"), bOverrideScreenPercentage));
+		if (bOverrideScreenPercentage)
+		{
+			EventAttributes.Add(FAnalyticsEventAttribute(TEXT("ScreenPercentage"), ScreenPercentage));
+		}
+		EventAttributes.Add(FAnalyticsEventAttribute(TEXT("AllowFinishCurrentFrame"), bAllowFinishCurrentFrame));
+#ifdef OVR_VISION_ENABLED
+		EventAttributes.Add(FAnalyticsEventAttribute(TEXT("HmdPosTracking"), bHmdPosTracking));
+		EventAttributes.Add(FAnalyticsEventAttribute(TEXT("LowPersistenceMode"), bLowPersistenceMode));
+#endif		
+		EventAttributes.Add(FAnalyticsEventAttribute(TEXT("UpdateOnRT"), bUpdateOnRT));
+
+		FString OutStr(TEXT("Editor.VR.DeviceInitialised"));
+		FEngineAnalytics::GetProvider().RecordEvent(OutStr, EventAttributes);
+	}
+}
+
 bool FOculusRiftHMD::IsPositionalTrackingEnabled() const
 {
 #ifdef OVR_VISION_ENABLED
diff --git a/Engine/Plugins/OculusRift/Source/OculusRift/Private/OculusRiftHMD.h b/Engine/Plugins/OculusRift/Source/OculusRift/Private/OculusRiftHMD.h
--- a/Engine/Plugins/OculusRift/Source/OculusRift/Private/OculusRiftHMD.h
+++ b/Engine/Plugins/OculusRift/Source/OculusRift/Private/OculusRiftHMD.h
@@ -81,6 +81,8 @@ public:
 	virtual bool Exec( UWorld* InWorld, const TCHAR* Cmd, FOutputDevice& Ar ) override;
 	virtual void OnScreenModeChange(EWindowMode::Type WindowMode) override;
 
+	virtual void RecordAnalytics() override;
+
 	/** IStereoRendering interface */
 	virtual bool IsStereoEnabled() const override;
 	virtual bool EnableStereo(bool stereo = true) override;
diff --git a/Engine/Source/Runtime/Engine/Classes/Engine/Engine.h b/Engine/Source/Runtime/Engine/Classes/Engine/Engine.h
--- a/Engine/Source/Runtime/Engine/Classes/Engine/Engine.h
+++ b/Engine/Source/Runtime/Engine/Classes/Engine/Engine.h
@@ -2146,6 +2146,11 @@ protected:
 	 */
 	virtual bool InitializeHMDDevice();
 
+	/**
+	 *	Record EngineAnalytics information for attached HMD devices
+	 */
+	virtual void RecordHMDAnalytics();
+
 	/**
 	 * Loads all Engine object references from their corresponding config entries.
 	 */
diff --git a/Engine/Source/Runtime/Engine/Private/UnrealEngine.cpp b/Engine/Source/Runtime/Engine/Private/UnrealEngine.cpp
--- a/Engine/Source/Runtime/Engine/Private/UnrealEngine.cpp
+++ b/Engine/Source/Runtime/Engine/Private/UnrealEngine.cpp
@@ -807,6 +807,9 @@ void UEngine::Init(IEngineLoop* InEngineLoop)
 		const FEngineStatFuncs& EngineStat = EngineStats[StatIdx];
 		NewStatDelegate.Broadcast(EngineStat.CommandName, EngineStat.CategoryName, EngineStat.DescriptionString);
 	}
+
+	// Record the analytics for any attached HMD devices
+	RecordHMDAnalytics();
 }
 
 void UEngine::RegisterBeginStreamingPauseRenderingDelegate( FBeginStreamingPauseDelegate* InDelegate )
@@ -1626,6 +1629,17 @@ bool UEngine::InitializeHMDDevice()
 	return StereoRenderingDevice.IsValid();
 }
 
+void UEngine::RecordHMDAnalytics()
+{
+	if( !GIsEditor )
+	{
+		if(HMDDevice.IsValid() && !FParse::Param(FCommandLine::Get(),TEXT("nohmd")))
+		{
+			HMDDevice->RecordAnalytics();
+		}
+	}
+}
+
 /** @return whether we're currently running in split screen (more than one local player) */
 bool UEngine::IsSplitScreen(UWorld *InWorld)
 {
diff --git a/Engine/Source/Runtime/HeadMountedDisplay/Public/IHeadMountedDisplay.h b/Engine/Source/Runtime/HeadMountedDisplay/Public/IHeadMountedDisplay.h
--- a/Engine/Source/Runtime/HeadMountedDisplay/Public/IHeadMountedDisplay.h
+++ b/Engine/Source/Runtime/HeadMountedDisplay/Public/IHeadMountedDisplay.h
@@ -231,6 +231,11 @@ public:
 	virtual const FTexture*  GetDistortionTexture() {return NULL;}	
 #endif
 
+	/**
+	 * Record analytics
+	 */
+	virtual void RecordAnalytics() {}
+
 private:
 	/** Stores the dimensions of the window before we moved into fullscreen mode, so they can be restored */
 	FSlateRect PreFullScreenRect;
