Restricting direct rendering in Oculus plugin to intended platforms

diff --git a/Engine/Plugins/OculusRift/Source/OculusRift/OculusRift.Build.cs b/Engine/Plugins/OculusRift/Source/OculusRift/OculusRift.Build.cs
--- a/Engine/Plugins/OculusRift/Source/OculusRift/OculusRift.Build.cs
+++ b/Engine/Plugins/OculusRift/Source/OculusRift/OculusRift.Build.cs
@@ -27,15 +27,19 @@ namespace UnrealBuildTool.Rules
 					"Renderer",
 					"ShaderCore",
 					"HeadMountedDisplay",
-					"D3D11RHI",
-					"OpenGLDrv"
 				}
 				);
 
-            // Currently, the Rift is only supported on windows platforms
+            // Currently, the Rift is only supported on windows and mac platforms
             if (Target.Platform == UnrealTargetPlatform.Win32 || Target.Platform == UnrealTargetPlatform.Win64 || Target.Platform == UnrealTargetPlatform.Mac)
             {
-				PrivateDependencyModuleNames.AddRange(new string[] { "LibOVR" });
+				PrivateDependencyModuleNames.AddRange(new string[] { "LibOVR", "OpenGLDrv" });
+
+                // Add direct rendering dependencies on a per-platform basis
+                if (Target.Platform == UnrealTargetPlatform.Win32 || Target.Platform == UnrealTargetPlatform.Win64)
+                {
+                    PrivateDependencyModuleNames.AddRange(new string[] { "D3D11RHI" });
+                }
             }
 		}
 	}
