#oculus Fix for non-direct rendering path and parallel rendering

diff --git a/Engine/Plugins/OculusRift/Source/OculusRift/Private/OculusRiftHMD.h b/Engine/Plugins/OculusRift/Source/OculusRift/Private/OculusRiftHMD.h
--- a/Engine/Plugins/OculusRift/Source/OculusRift/Private/OculusRiftHMD.h
+++ b/Engine/Plugins/OculusRift/Source/OculusRift/Private/OculusRiftHMD.h
@@ -237,7 +237,7 @@ public:
 
 #else
 
-	virtual void FinishRenderingFrame_RenderThread() override;
+virtual void FinishRenderingFrame_RenderThread(FRHICommandListImmediate& RHICmdList) override;
 
 #endif // #ifdef OVR_DIRECT_RENDERING
 
diff --git a/Engine/Plugins/OculusRift/Source/OculusRift/Private/OculusRiftRender.cpp b/Engine/Plugins/OculusRift/Source/OculusRift/Private/OculusRiftRender.cpp
--- a/Engine/Plugins/OculusRift/Source/OculusRift/Private/OculusRiftRender.cpp
+++ b/Engine/Plugins/OculusRift/Source/OculusRift/Private/OculusRiftRender.cpp
@@ -332,14 +332,14 @@ bool FOculusRiftHMD::NeedReAllocateViewportRenderTarget(const FViewport& Viewpor
 
 #else // no direct rendering
 
-void FOculusRiftHMD::FinishRenderingFrame_RenderThread()
+void FOculusRiftHMD::FinishRenderingFrame_RenderThread(FRHICommandListImmediate& RHICmdList)
 {
 	if (RenderParams_RenderThread.bFrameBegun)
 	{
 		check(IsInRenderingThread());
 		if (RenderParams_RenderThread.bTimeWarp)
 		{
-			RHIBlockUntilGPUIdle();
+			RHICmdList.BlockUntilGPUIdle();
 		}
 		ovrHmd_EndFrameTiming(Hmd);
 		RenderParams_RenderThread.bFrameBegun = false;
diff --git a/Engine/Source/Runtime/Engine/Public/StereoRendering.h b/Engine/Source/Runtime/Engine/Public/StereoRendering.h
--- a/Engine/Source/Runtime/Engine/Public/StereoRendering.h
+++ b/Engine/Source/Runtime/Engine/Public/StereoRendering.h
@@ -87,5 +87,5 @@ public:
 	/**
 	 * Called after Present is called.
 	 */
-	virtual void FinishRenderingFrame_RenderThread() {}
+	virtual void FinishRenderingFrame_RenderThread(FRHICommandListImmediate& RHICmdList) {}
 };
diff --git a/Engine/Source/Runtime/SlateRHIRenderer/Private/SlateRHIRenderer.cpp b/Engine/Source/Runtime/SlateRHIRenderer/Private/SlateRHIRenderer.cpp
--- a/Engine/Source/Runtime/SlateRHIRenderer/Private/SlateRHIRenderer.cpp
+++ b/Engine/Source/Runtime/SlateRHIRenderer/Private/SlateRHIRenderer.cpp
@@ -467,7 +467,7 @@ void FSlateRHIRenderer::DrawWindow_RenderThread(FRHICommandListImmediate& RHICmd
 
 	if (bNeedCallFinishFrameForStereo)
 	{
-		GEngine->StereoRenderingDevice->FinishRenderingFrame_RenderThread();
+		GEngine->StereoRenderingDevice->FinishRenderingFrame_RenderThread(RHICmdList);
 	}
 	uint32 EndTime		= FPlatformTime::Cycles();
 
