#oculus More fixes for direct render path in Oculus

diff --git a/Engine/Plugins/OculusRift/Source/OculusRift/Private/OculusRiftRender.cpp b/Engine/Plugins/OculusRift/Source/OculusRift/Private/OculusRiftRender.cpp
--- a/Engine/Plugins/OculusRift/Source/OculusRift/Private/OculusRiftRender.cpp
+++ b/Engine/Plugins/OculusRift/Source/OculusRift/Private/OculusRiftRender.cpp
@@ -97,9 +97,10 @@ void FOculusRiftHMD::DrawDistortionMesh_RenderThread(FRenderingCompositePassCont
 	FMatrix QuadPosTransform = FMatrix::Identity;
 	const FIntRect SrcRect = View.ViewRect;
 
+	FRHICommandListImmediate& RHICmdList = Context.RHICmdList;
 	const FSceneViewFamily& ViewFamily = *(View.Family);
 	FIntPoint ViewportSize = ViewFamily.RenderTarget->GetSizeXY();
-	RHISetViewport(0, 0, 0.0f, ViewportSize.X, ViewportSize.Y, 1.0f);
+	RHICmdList.SetViewport(0, 0, 0.0f, ViewportSize.X, ViewportSize.Y, 1.0f);
 
 	Ptr<FDistortionMesh> mesh = RenderParams_RenderThread.pDistortionMesh[(View.StereoPass == eSSP_LEFT_EYE) ? 0 : 1];
 
