#oculus More fixes for virtual thunk on direct rendering path

diff --git a/Engine/Plugins/OculusRift/Source/OculusRift/Private/OculusRiftHMD.h b/Engine/Plugins/OculusRift/Source/OculusRift/Private/OculusRiftHMD.h
--- a/Engine/Plugins/OculusRift/Source/OculusRift/Private/OculusRiftHMD.h
+++ b/Engine/Plugins/OculusRift/Source/OculusRift/Private/OculusRiftHMD.h
@@ -100,10 +100,12 @@ public:
 	virtual void GetEyeRenderParams_RenderThread(EStereoscopicPass StereoPass, FVector2D& EyeToSrcUVScaleValue, FVector2D& EyeToSrcUVOffsetValue) const override;
 	virtual void GetTimewarpMatrices_RenderThread(EStereoscopicPass StereoPass, FMatrix& EyeRotationStart, FMatrix& EyeRotationEnd) const override;
 
+#ifdef OVR_DIRECT_RENDERING
 	virtual void UpdateViewport(bool bUseSeparateRenderTarget, const FViewport& Viewport) override;
-
 	virtual void CalculateRenderTargetSize(uint32& InOutSizeX, uint32& InOutSizeY) const override;
 	virtual bool NeedReAllocateViewportRenderTarget(const FViewport& Viewport) const override;
+#endif//OVR_DIRECT_RENDERING
+
 	virtual bool ShouldUseSeparateRenderTarget() const override
 	{
 		check(IsInGameThread());
