#TTP 340946: OCULUS: MAC: Going fullscreen renders black

diff --git a/Engine/Plugins/OculusRift/Source/OculusRift/Private/OculusRiftHMD.h b/Engine/Plugins/OculusRift/Source/OculusRift/Private/OculusRiftHMD.h
--- a/Engine/Plugins/OculusRift/Source/OculusRift/Private/OculusRiftHMD.h
+++ b/Engine/Plugins/OculusRift/Source/OculusRift/Private/OculusRiftHMD.h
@@ -108,8 +108,12 @@ public:
 
 	virtual bool ShouldUseSeparateRenderTarget() const override
 	{
+#ifdef OVR_DIRECT_RENDERING
 		check(IsInGameThread());
 		return IsStereoEnabled();
+#else
+		return false;
+#endif//OVR_DIRECT_RENDERING
 	}
 
     /** ISceneViewExtension interface */
