#oculus Merging fix for vsync overrides

diff --git a/Engine/Plugins/OculusRift/Source/OculusRift/Private/OculusRiftHMD.cpp b/Engine/Plugins/OculusRift/Source/OculusRift/Private/OculusRiftHMD.cpp
--- a/Engine/Plugins/OculusRift/Source/OculusRift/Private/OculusRiftHMD.cpp
+++ b/Engine/Plugins/OculusRift/Source/OculusRift/Private/OculusRiftHMD.cpp
@@ -181,7 +181,7 @@ bool FOculusRiftHMD::IsInLowPersistenceMode() const
 void FOculusRiftHMD::EnableLowPersistenceMode(bool Enable)
 {
 	bLowPersistenceMode = Enable;
-	UpdateSensorHmdCaps();
+	UpdateHmdCaps();
 }
 
 float FOculusRiftHMD::GetInterpupillaryDistance() const
@@ -569,7 +569,7 @@ bool FOculusRiftHMD::Exec( UWorld* InWorld, const TCHAR* Cmd, FOutputDevice& Ar
 			{
 				bLowPersistenceMode = !bLowPersistenceMode;
 			}
-			UpdateSensorHmdCaps();
+			UpdateHmdCaps();
 			Ar.Logf(TEXT("Low Persistence is currently %s"), (bLowPersistenceMode) ? TEXT("ON") : TEXT("OFF"));
 			return true;
 		}
@@ -686,13 +686,13 @@ bool FOculusRiftHMD::Exec( UWorld* InWorld, const TCHAR* Cmd, FOutputDevice& Ar
 		if (FParse::Command(&Cmd, TEXT("ON")))
 		{
 			bYawDriftCorrectionEnabled = true;
-			UpdateSensorHmdCaps();
+			UpdateHmdCaps();
 			return true;
 		}
 		else if (FParse::Command(&Cmd, TEXT("OFF")))
 		{
 			bYawDriftCorrectionEnabled = false;
-			UpdateSensorHmdCaps();
+			UpdateHmdCaps();
 			return true;
 		}
 		else if (FParse::Command(&Cmd, TEXT("SHOW")))
@@ -763,21 +763,21 @@ bool FOculusRiftHMD::Exec( UWorld* InWorld, const TCHAR* Cmd, FOutputDevice& Ar
 		else if (FParse::Command(&Cmd, TEXT("ON")) || FParse::Command(&Cmd, TEXT("ENABLE")))
 		{
 			bHmdPosTracking = true;
-			UpdateSensorHmdCaps();
+			UpdateHmdCaps();
 			return true;
 		}
 		else if (FParse::Command(&Cmd, TEXT("OFF")) || FParse::Command(&Cmd, TEXT("DISABLE")))
 		{
 			bHmdPosTracking = false;
 			bHaveVisionTracking = false;
-			UpdateSensorHmdCaps();
+			UpdateHmdCaps();
 			return true;
 		}
 		else if (FParse::Command(&Cmd, TEXT("TOGGLE")))
 		{
 			bHmdPosTracking = !bHmdPosTracking;
 			bHaveVisionTracking = false;
-			UpdateSensorHmdCaps();
+			UpdateHmdCaps();
 			return true;
 		}
 #if !UE_BUILD_SHIPPING
@@ -1019,14 +1019,13 @@ void FOculusRiftHMD::ApplySystemOverridesOnStereo(bool bForce)
 		{
 			static IConsoleVariable* CVSyncVar = IConsoleManager::Get().FindConsoleVariable(TEXT("r.VSync"));
 			bVSync = CVSyncVar->GetInt() != 0;
-
-			#ifdef OVR_DIRECT_RENDERING 
-			GetActiveRHIBridgeImpl()->SetNeedReinitRendererAPI();
-			#endif
 		}
+		UpdateHmdCaps();
 
+#ifndef OVR_DIRECT_RENDERING
 		static IConsoleVariable* CFinishFrameVar = IConsoleManager::Get().FindConsoleVariable(TEXT("r.FinishCurrentFrame"));
 		CFinishFrameVar->Set(bAllowFinishCurrentFrame);
+#endif
 	}
 }
 
@@ -1423,7 +1422,7 @@ void FOculusRiftHMD::Startup()
 	LoadFromIni();
 	SaveSystemValues();
 
-	UpdateSensorHmdCaps();
+	UpdateHmdCaps();
 
 #ifdef OVR_DIRECT_RENDERING
 #if defined(OVR_D3D_VERSION) && (OVR_D3D_VERSION == 11)
@@ -1477,7 +1476,7 @@ void FOculusRiftHMD::Shutdown()
 	UE_LOG(LogHMD, Log, TEXT("Oculus shutdown."));
 }
 
-void FOculusRiftHMD::UpdateSensorHmdCaps()
+void FOculusRiftHMD::UpdateHmdCaps()
 {
 	if (Hmd)
 	{
@@ -1508,6 +1507,16 @@ void FOculusRiftHMD::UpdateSensorHmdCaps()
 			HmdCaps &= ~ovrHmdCap_LowPersistence;
 		}
 		HmdCaps |= ovrHmdCap_LatencyTest;
+
+		if (bVSync) 
+		{
+			HmdCaps &= ~ovrHmdCap_NoVSync;
+		}
+		else
+		{
+			HmdCaps |= ovrHmdCap_NoVSync;
+		}
+
 		ovrHmd_SetEnabledCaps(Hmd, HmdCaps);
 
 		ovrHmd_StartSensor(Hmd, SensorCaps, 0);
diff --git a/Engine/Plugins/OculusRift/Source/OculusRift/Private/OculusRiftHMD.h b/Engine/Plugins/OculusRift/Source/OculusRift/Private/OculusRiftHMD.h
--- a/Engine/Plugins/OculusRift/Source/OculusRift/Private/OculusRiftHMD.h
+++ b/Engine/Plugins/OculusRift/Source/OculusRift/Private/OculusRiftHMD.h
@@ -373,7 +373,7 @@ private:
 
 	void PrecalculatePostProcess_NoLock();
 
-	void UpdateSensorHmdCaps();
+	void UpdateHmdCaps();
 
 	void PoseToOrientationAndPosition(const ovrPosef& InPose, FQuat& OutOrientation, FVector& OutPosition) const;
 
diff --git a/Engine/Plugins/OculusRift/Source/OculusRift/Private/OculusRiftRender.cpp b/Engine/Plugins/OculusRift/Source/OculusRift/Private/OculusRiftRender.cpp
--- a/Engine/Plugins/OculusRift/Source/OculusRift/Private/OculusRiftRender.cpp
+++ b/Engine/Plugins/OculusRift/Source/OculusRift/Private/OculusRiftRender.cpp
@@ -633,7 +633,6 @@ void FOculusRiftHMD::D3D11Bridge::BeginRendering()
 			OVR::Lock::Locker lock(&ModifyLock);
 			check(Cfg.D3D11.pSwapChain); // make sure Config is initialized
 			(Plugin->bTimeWarp) ? Plugin->DistortionCaps |= ovrDistortionCap_TimeWarp : Plugin->DistortionCaps &= ~ovrDistortionCap_TimeWarp;
-			(Plugin->bVSync) ? Plugin->HmdCaps &= ~ovrHmdCap_NoVSync : Plugin->HmdCaps |= ovrHmdCap_NoVSync;
 			if (!ovrHmd_ConfigureRendering(Plugin->Hmd, &Cfg.Config, Plugin->DistortionCaps, 
 				Plugin->RenderParams_RenderThread.EyeFov, Plugin->RenderParams_RenderThread.EyeRenderDesc))
 			{
@@ -931,7 +930,6 @@ void FOculusRiftHMD::OGLBridge::BeginRendering()
 			Plugin->DistortionCaps &= ~ovrDistortionCap_SRGB;
 			Plugin->DistortionCaps |= ovrDistortionCap_FlipInput;
 			(Plugin->bTimeWarp) ? Plugin->DistortionCaps |= ovrDistortionCap_TimeWarp : Plugin->DistortionCaps &= ~ovrDistortionCap_TimeWarp;
-			(Plugin->bVSync) ? Plugin->HmdCaps &= ~ovrHmdCap_NoVSync : Plugin->HmdCaps |= ovrHmdCap_NoVSync;
 			if (!ovrHmd_ConfigureRendering(Plugin->Hmd, &Cfg.Config, Plugin->DistortionCaps, 
 				Plugin->RenderParams_RenderThread.EyeFov, Plugin->RenderParams_RenderThread.EyeRenderDesc))
 			{
