Oculus HMD analytics were lost in a merge, added back in.
Removed some orphaned head model stuff.

#ttp 342873

diff --git a/Engine/Plugins/OculusRift/Source/OculusRift/Private/OculusRiftHMD.cpp b/Engine/Plugins/OculusRift/Source/OculusRift/Private/OculusRiftHMD.cpp
--- a/Engine/Plugins/OculusRift/Source/OculusRift/Private/OculusRiftHMD.cpp
+++ b/Engine/Plugins/OculusRift/Source/OculusRift/Private/OculusRiftHMD.cpp
@@ -2,6 +2,8 @@
 
 #include "OculusRiftPrivate.h"
 #include "OculusRiftHMD.h"
+#include "EngineAnalytics.h"
+#include "Runtime/Analytics/Analytics/Public/Interfaces/IAnalyticsProvider.h"
 
 #define DEFAULT_PREDICTION_IN_SECONDS 0.035
 
@@ -933,6 +935,59 @@ void FOculusRiftHMD::OnScreenModeChange(EWindowMode::Type WindowMode)
 	UpdateStereoRenderingParams();
 }
 
+void FOculusRiftHMD::RecordAnalytics()
+{
+	if (FEngineAnalytics::IsAvailable())
+	{
+		// prepare and send analytics data
+		TArray<FAnalyticsEventAttribute> EventAttributes;
+
+		IHeadMountedDisplay::MonitorInfo MonitorInfo;
+		GetHMDMonitorInfo(MonitorInfo);
+		EventAttributes.Add(FAnalyticsEventAttribute(TEXT("DeviceName"), MonitorInfo.MonitorName));
+		EventAttributes.Add(FAnalyticsEventAttribute(TEXT("DisplayId"), MonitorInfo.MonitorId));
+		FString MonResolution(FString::Printf(TEXT("(%d, %d)"), MonitorInfo.ResolutionX, MonitorInfo.ResolutionY));
+		EventAttributes.Add(FAnalyticsEventAttribute(TEXT("Resolution"), MonResolution));
+
+		EventAttributes.Add(FAnalyticsEventAttribute(TEXT("ChromaAbCorrectionEnabled"), bChromaAbCorrectionEnabled));
+		EventAttributes.Add(FAnalyticsEventAttribute(TEXT("MagEnabled"), bYawDriftCorrectionEnabled));
+		EventAttributes.Add(FAnalyticsEventAttribute(TEXT("DevSettingsEnabled"), bDevSettingsEnabled));
+		EventAttributes.Add(FAnalyticsEventAttribute(TEXT("MotionPredictionEnabled"), (MotionPredictionInSeconds > 0)));
+		EventAttributes.Add(FAnalyticsEventAttribute(TEXT("TiltCorrectionEnabled"), bTiltCorrectionEnabled));
+		EventAttributes.Add(FAnalyticsEventAttribute(TEXT("AccelGain"), AccelGain));
+		EventAttributes.Add(FAnalyticsEventAttribute(TEXT("OverrideInterpupillaryDistance"), bOverrideIPD));
+		if (bOverrideIPD)
+		{
+			EventAttributes.Add(FAnalyticsEventAttribute(TEXT("InterpupillaryDistance"), GetInterpupillaryDistance()));
+		}
+		EventAttributes.Add(FAnalyticsEventAttribute(TEXT("OverrideStereo"), bOverrideStereo));
+		if (bOverrideStereo)
+		{
+			EventAttributes.Add(FAnalyticsEventAttribute(TEXT("HFOV"), HFOVInRadians));
+			EventAttributes.Add(FAnalyticsEventAttribute(TEXT("VFOV"), VFOVInRadians));
+		}
+		EventAttributes.Add(FAnalyticsEventAttribute(TEXT("OverrideVSync"), bOverrideVSync));
+		if (bOverrideVSync)
+		{
+			EventAttributes.Add(FAnalyticsEventAttribute(TEXT("VSync"), bVSync));
+		}
+		EventAttributes.Add(FAnalyticsEventAttribute(TEXT("OverrideScreenPercentage"), bOverrideScreenPercentage));
+		if (bOverrideScreenPercentage)
+		{
+			EventAttributes.Add(FAnalyticsEventAttribute(TEXT("ScreenPercentage"), ScreenPercentage));
+		}
+		EventAttributes.Add(FAnalyticsEventAttribute(TEXT("AllowFinishCurrentFrame"), bAllowFinishCurrentFrame));
+#ifdef OVR_VISION_ENABLED
+		EventAttributes.Add(FAnalyticsEventAttribute(TEXT("HmdPosTracking"), bHmdPosTracking));
+		EventAttributes.Add(FAnalyticsEventAttribute(TEXT("LowPersistenceMode"), bLowPersistenceMode));
+#endif		
+		EventAttributes.Add(FAnalyticsEventAttribute(TEXT("UpdateOnRT"), bUpdateOnRT));
+
+		FString OutStr(TEXT("Editor.VR.DeviceInitialised"));
+		FEngineAnalytics::GetProvider().RecordEvent(OutStr, EventAttributes);
+	}
+}
+
 bool FOculusRiftHMD::IsPositionalTrackingEnabled() const
 {
 #ifdef OVR_VISION_ENABLED
diff --git a/Engine/Plugins/OculusRift/Source/OculusRift/Private/OculusRiftHMD.h b/Engine/Plugins/OculusRift/Source/OculusRift/Private/OculusRiftHMD.h
--- a/Engine/Plugins/OculusRift/Source/OculusRift/Private/OculusRiftHMD.h
+++ b/Engine/Plugins/OculusRift/Source/OculusRift/Private/OculusRiftHMD.h
@@ -87,6 +87,8 @@ public:
 	virtual bool Exec( UWorld* InWorld, const TCHAR* Cmd, FOutputDevice& Ar ) override;
 	virtual void OnScreenModeChange(EWindowMode::Type WindowMode) override;
 
+	virtual void RecordAnalytics() override;
+
 	/** IStereoRendering interface */
 	virtual bool IsStereoEnabled() const override;
 	virtual bool EnableStereo(bool stereo = true) override;
@@ -366,11 +368,6 @@ private:
 
 	void ResetControlRotation() const;
 
-	/** Get/set head model. Units are meters, not UU! 
-	    Use ToFVector and ToFVector_M2U conversion routines. */
-	OVR::Vector3f	GetHeadModel() const;
-	void			SetHeadModel(const OVR::Vector3f&);
-
 	void PrecalculatePostProcess_NoLock();
 
 	void UpdateHmdCaps();
@@ -458,9 +455,6 @@ private: // data
 	/** Gain for gravity correction (should not need to be changed) */
 	float AccelGain;
 
-	/** Whether or not to use a head model offset to determine view position */
-	bool bUseHeadModel;
-
     /** Distortion on/off */
     bool bHmdDistortion;
 
