The Occulus monitor ID is not a string on OS X, so don't treat it like one when sending analytics.

diff --git a/Engine/Plugins/OculusRift/Source/OculusRift/Private/OculusRiftHMD.cpp b/Engine/Plugins/OculusRift/Source/OculusRift/Private/OculusRiftHMD.cpp
--- a/Engine/Plugins/OculusRift/Source/OculusRift/Private/OculusRiftHMD.cpp
+++ b/Engine/Plugins/OculusRift/Source/OculusRift/Private/OculusRiftHMD.cpp
@@ -945,7 +945,12 @@ void FOculusRiftHMD::RecordAnalytics()
 		IHeadMountedDisplay::MonitorInfo MonitorInfo;
 		GetHMDMonitorInfo(MonitorInfo);
 		EventAttributes.Add(FAnalyticsEventAttribute(TEXT("DeviceName"), MonitorInfo.MonitorName));
+#if PLATFORM_MAC // On OS X MonitorId is the CGDirectDisplayID aka uint64, not a string
+		FString DisplayId(FString::Printf(TEXT("%llu"), MonitorInfo.MonitorId));
+		EventAttributes.Add(FAnalyticsEventAttribute(TEXT("DisplayId"), DisplayId));
+#else
 		EventAttributes.Add(FAnalyticsEventAttribute(TEXT("DisplayId"), MonitorInfo.MonitorId));
+#endif
 		FString MonResolution(FString::Printf(TEXT("(%d, %d)"), MonitorInfo.ResolutionX, MonitorInfo.ResolutionY));
 		EventAttributes.Add(FAnalyticsEventAttribute(TEXT("Resolution"), MonResolution));
 
