#analytics Adding in device name to the analytics info

diff --git a/Engine/Plugins/OculusRift/Source/OculusRift/Private/OculusRiftHMD.cpp b/Engine/Plugins/OculusRift/Source/OculusRift/Private/OculusRiftHMD.cpp
--- a/Engine/Plugins/OculusRift/Source/OculusRift/Private/OculusRiftHMD.cpp
+++ b/Engine/Plugins/OculusRift/Source/OculusRift/Private/OculusRiftHMD.cpp
@@ -944,7 +944,8 @@ void FOculusRiftHMD::RecordAnalytics()
 
 		IHeadMountedDisplay::MonitorInfo MonitorInfo;
 		GetHMDMonitorInfo(MonitorInfo);
-		EventAttributes.Add(FAnalyticsEventAttribute(TEXT("DeviceName"), MonitorInfo.MonitorName));
+		EventAttributes.Add(FAnalyticsEventAttribute(TEXT("DeviceName"), TEXT("OculusRift")));
+		EventAttributes.Add(FAnalyticsEventAttribute(TEXT("DisplayDeviceName"), MonitorInfo.MonitorName));
 #if PLATFORM_MAC // On OS X MonitorId is the CGDirectDisplayID aka uint64, not a string
 		FString DisplayId(FString::Printf(TEXT("%llu"), MonitorInfo.MonitorId));
 		EventAttributes.Add(FAnalyticsEventAttribute(TEXT("DisplayId"), DisplayId));
