- modified oculus hmd analytics to report manufacturer and device type strings

diff --git a/Engine/Plugins/OculusRift/Source/OculusRift/Private/OculusRiftHMD.cpp b/Engine/Plugins/OculusRift/Source/OculusRift/Private/OculusRiftHMD.cpp
--- a/Engine/Plugins/OculusRift/Source/OculusRift/Private/OculusRiftHMD.cpp
+++ b/Engine/Plugins/OculusRift/Source/OculusRift/Private/OculusRiftHMD.cpp
@@ -944,7 +944,7 @@ void FOculusRiftHMD::RecordAnalytics()
 
 		IHeadMountedDisplay::MonitorInfo MonitorInfo;
 		GetHMDMonitorInfo(MonitorInfo);
-		EventAttributes.Add(FAnalyticsEventAttribute(TEXT("DeviceName"), TEXT("OculusRift")));
+		EventAttributes.Add(FAnalyticsEventAttribute(TEXT("DeviceName"), FString::Printf(TEXT("%s - %s"), ANSI_TO_TCHAR(HmdDesc.Manufacturer), ANSI_TO_TCHAR(HmdDesc.ProductName))));
 		EventAttributes.Add(FAnalyticsEventAttribute(TEXT("DisplayDeviceName"), MonitorInfo.MonitorName));
 #if PLATFORM_MAC // On OS X MonitorId is the CGDirectDisplayID aka uint64, not a string
 		FString DisplayId(FString::Printf(TEXT("%llu"), MonitorInfo.MonitorId));
