Fix for oculus causing an InitializeCriticalSectionEx exception on windows XP

TTP# 343103 - COMPAT: Kernel32.dll error is present when trying to run a packaged project on Windows XP.

diff --git a/Engine/Plugins/OculusRift/Source/OculusRift/Private/OculusRiftHMD.cpp b/Engine/Plugins/OculusRift/Source/OculusRift/Private/OculusRiftHMD.cpp
--- a/Engine/Plugins/OculusRift/Source/OculusRift/Private/OculusRiftHMD.cpp
+++ b/Engine/Plugins/OculusRift/Source/OculusRift/Private/OculusRiftHMD.cpp
@@ -16,30 +16,6 @@
 #if PLATFORM_SUPPORTS_PRAGMA_PACK
 #pragma pack (pop)
 #endif
-//////////////////////////////////////////////////////////////////////////
-class OculusLog : public OVR::Log
-{
-public:
-	OculusLog()
-	{
-		SetLoggingMask(OVR::LogMask_Debug | OVR::LogMask_Regular);
-	}
-
-	// This virtual function receives all the messages,
-	// developers should override this function in order to do custom logging
-	virtual void    LogMessageVarg(LogMessageType messageType, const char* fmt, va_list argList)
-	{
-		if ((messageType & GetLoggingMask()) == 0)
-			return;
-
-		ANSICHAR buf[1024];
-		int32 len = FCStringAnsi::GetVarArgs(buf, sizeof(buf), sizeof(buf)/sizeof(ANSICHAR), fmt, argList);
-		if (len > 0 && buf[len - 1] == '\n') // truncate the trailing new-line char, since Logf adds its own
-			buf[len - 1] = '\0';
-		TCHAR* tbuf = ANSI_TO_TCHAR(buf);
-		GLog->Logf(TEXT("OCULUS: %s"), tbuf);
-	}
-};
 #endif // #if !UE_BUILD_SHIPPING
 
 //---------------------------------------------------
@@ -66,6 +42,40 @@ TSharedPtr< class IHeadMountedDisplay > FOculusRiftPlugin::CreateHeadMountedDisp
 	return NULL;
 }
 
+//---------------------------------------------------
+// Oculus Rift IHeadMountedDisplay Implementation
+//---------------------------------------------------
+
+#if OCULUS_RIFT_SUPPORTED_PLATFORMS
+
+#if !UE_BUILD_SHIPPING
+//////////////////////////////////////////////////////////////////////////
+class OculusLog : public OVR::Log
+{
+public:
+	OculusLog()
+	{
+		SetLoggingMask(OVR::LogMask_Debug | OVR::LogMask_Regular);
+	}
+
+	// This virtual function receives all the messages,
+	// developers should override this function in order to do custom logging
+	virtual void    LogMessageVarg(LogMessageType messageType, const char* fmt, va_list argList)
+	{
+		if ((messageType & GetLoggingMask()) == 0)
+			return;
+
+		ANSICHAR buf[1024];
+		int32 len = FCStringAnsi::GetVarArgs(buf, sizeof(buf), sizeof(buf) / sizeof(ANSICHAR), fmt, argList);
+		if (len > 0 && buf[len - 1] == '\n') // truncate the trailing new-line char, since Logf adds its own
+			buf[len - 1] = '\0';
+		TCHAR* tbuf = ANSI_TO_TCHAR(buf);
+		GLog->Logf(TEXT("OCULUS: %s"), tbuf);
+	}
+};
+
+#endif
+
 //////////////////////////////////////////////////////////////////////////
 class ConditionalLocker
 {
@@ -86,12 +96,6 @@ private:
 	OVR::Lock*	pLock;
 };
 
-//---------------------------------------------------
-// Oculus Rift IHeadMountedDisplay Implementation
-//---------------------------------------------------
-
-#if OCULUS_RIFT_SUPPORTED_PLATFORMS
-
 bool FOculusRiftHMD::IsHMDEnabled() const
 {
 	return bHMDEnabled;
diff --git a/Engine/Plugins/OculusRift/Source/OculusRift/Private/OculusRiftRender.cpp b/Engine/Plugins/OculusRift/Source/OculusRift/Private/OculusRiftRender.cpp
--- a/Engine/Plugins/OculusRift/Source/OculusRift/Private/OculusRiftRender.cpp
+++ b/Engine/Plugins/OculusRift/Source/OculusRift/Private/OculusRiftRender.cpp
@@ -3,6 +3,8 @@
 #include "OculusRiftPrivate.h"
 #include "OculusRiftHMD.h"
 
+#if OCULUS_RIFT_SUPPORTED_PLATFORMS
+
 #ifdef __clang__ // @todo: ThirdParty/Oculus/LibOVR/Include/../Src/OVR_DeviceMessages.h:227:46: warning: field 'CameraHandle' will be initialized after field 'pFrameData' [-Wreorder]
 	#pragma clang diagnostic push
 	#pragma clang diagnostic ignored "-Wreorder"
@@ -1074,3 +1076,4 @@ bool FOculusRiftHMD::OGLBridge::Present(int SyncInterval)
 #endif // OVR_DIRECT_RENDERING
 
 
+#endif // OCULUS_RIFT_SUPPORTED_PLATFORMS
\ No newline at end of file
diff --git a/Engine/Plugins/OculusRift/Source/OculusRift/Public/IOculusRiftPlugin.h b/Engine/Plugins/OculusRift/Source/OculusRift/Public/IOculusRiftPlugin.h
--- a/Engine/Plugins/OculusRift/Source/OculusRift/Public/IOculusRiftPlugin.h
+++ b/Engine/Plugins/OculusRift/Source/OculusRift/Public/IOculusRiftPlugin.h
@@ -5,7 +5,8 @@
 #include "ModuleManager.h"
 #include "IHeadMountedDisplayModule.h"
 
-#define OCULUS_RIFT_SUPPORTED_PLATFORMS PLATFORM_WINDOWS || PLATFORM_MAC
+// Oculus support is not available on windows xp
+#define OCULUS_RIFT_SUPPORTED_PLATFORMS (PLATFORM_WINDOWS && WINVER > 0x0502) || PLATFORM_MAC
 
 /**
  * The public interface to this module.  In most cases, this interface is only public to sibling modules 
