Removed no longer needed clang warning workaround.

diff --git a/Engine/Plugins/OculusRift/Source/OculusRift/Private/OculusRiftRender.cpp b/Engine/Plugins/OculusRift/Source/OculusRift/Private/OculusRiftRender.cpp
--- a/Engine/Plugins/OculusRift/Source/OculusRift/Private/OculusRiftRender.cpp
+++ b/Engine/Plugins/OculusRift/Source/OculusRift/Private/OculusRiftRender.cpp
@@ -5,17 +5,8 @@
 
 #if OCULUS_RIFT_SUPPORTED_PLATFORMS
 
-#ifdef __clang__ // @todo: ThirdParty/Oculus/LibOVR/Include/../Src/OVR_DeviceMessages.h:227:46: warning: field 'CameraHandle' will be initialized after field 'pFrameData' [-Wreorder]
-	#pragma clang diagnostic push
-	#pragma clang diagnostic ignored "-Wreorder"
-#endif
-
 #include "../Src/OVR_Stereo.h"
 
-#ifdef __clang__
-	#pragma clang diagnostic pop
-#endif
-
 #include "RendererPrivate.h"
 #include "ScenePrivate.h"
 #include "PostProcess/PostProcessHMD.h"
