Fix for oculus xp breakage.

diff --git a/Engine/Plugins/OculusRift/Source/OculusRift/Private/OculusRiftHMD.cpp b/Engine/Plugins/OculusRift/Source/OculusRift/Private/OculusRiftHMD.cpp
--- a/Engine/Plugins/OculusRift/Source/OculusRift/Private/OculusRiftHMD.cpp
+++ b/Engine/Plugins/OculusRift/Source/OculusRift/Private/OculusRiftHMD.cpp
@@ -45,7 +45,9 @@ TSharedPtr< class IHeadMountedDisplay > FOculusRiftPlugin::CreateHeadMountedDisp
 
 void FOculusRiftPlugin::PreInit()
 {
+#if OCULUS_RIFT_SUPPORTED_PLATFORMS
 	FOculusRiftHMD::PreInit();
+#endif//OCULUS_RIFT_SUPPORTED_PLATFORMS
 }
 
 //---------------------------------------------------
