#ue4 Merging fix to main for Oculus positional tracking

diff --git a/Engine/Plugins/OculusRift/Source/OculusRift/Private/OculusRiftHMD.cpp b/Engine/Plugins/OculusRift/Source/OculusRift/Private/OculusRiftHMD.cpp
--- a/Engine/Plugins/OculusRift/Source/OculusRift/Private/OculusRiftHMD.cpp
+++ b/Engine/Plugins/OculusRift/Source/OculusRift/Private/OculusRiftHMD.cpp
@@ -244,6 +244,8 @@ void FOculusRiftHMD::GetCurrentOrientationAndPosition(FQuat& CurrentOrientation,
 	checkf(IsInGameThread());
 	GetCurrentPose(CurHmdOrientation, CurHmdPosition);
 	CurrentOrientation = LastHmdOrientation = CurHmdOrientation;
+
+	CurrentPosition = CurHmdPosition;
 	
 	// @TODO: we can't actually return just CurHmdPosition here, it should be rotated by CurrentOrientation.
 	// Another problem: if we return rotated position here, we shouldn't apply position in CalculateStereoViewOffset, otherwise the
