Moved the Oculus Rift plugin into the Runtime folder.

diff --git a/Engine/Plugins/OculusRift/OculusRift.uplugin b/Engine/Plugins/OculusRift/OculusRift.uplugin
deleted file mode 100644
diff --git a/Engine/Plugins/OculusRift/Source/OculusRift/OculusRift.Build.cs b/Engine/Plugins/OculusRift/Source/OculusRift/OculusRift.Build.cs
deleted file mode 100644
diff --git a/Engine/Plugins/OculusRift/Source/OculusRift/Private/OculusRiftHMD.cpp b/Engine/Plugins/OculusRift/Source/OculusRift/Private/OculusRiftHMD.cpp
deleted file mode 100644
diff --git a/Engine/Plugins/OculusRift/Source/OculusRift/Private/OculusRiftHMD.h b/Engine/Plugins/OculusRift/Source/OculusRift/Private/OculusRiftHMD.h
deleted file mode 100644
diff --git a/Engine/Plugins/OculusRift/Source/OculusRift/Private/OculusRiftPrivate.h b/Engine/Plugins/OculusRift/Source/OculusRift/Private/OculusRiftPrivate.h
deleted file mode 100644
diff --git a/Engine/Plugins/OculusRift/Source/OculusRift/Private/OculusRiftRender.cpp b/Engine/Plugins/OculusRift/Source/OculusRift/Private/OculusRiftRender.cpp
deleted file mode 100644
diff --git a/Engine/Plugins/OculusRift/Source/OculusRift/Public/IOculusRiftPlugin.h b/Engine/Plugins/OculusRift/Source/OculusRift/Public/IOculusRiftPlugin.h
deleted file mode 100644
diff --git a/Engine/Plugins/Runtime/OculusRift/OculusRift.uplugin b/Engine/Plugins/Runtime/OculusRift/OculusRift.uplugin
new file mode 100644
--- /dev/null
+++ b/Engine/Plugins/Runtime/OculusRift/OculusRift.uplugin
@@ -0,0 +1,23 @@
+{
+	"FileVersion" : 3,
+
+	"FriendlyName" : "Oculus Rift Plugin",
+	"Version" : 1,
+	"VersionName" : "1.0",
+	"CreatedBy" : "Epic Games, Inc.",
+	"CreatedByURL" : "http://epicgames.com",
+	"EngineVersion" : "4.2.0",
+	"Description" : "Support for the Oculus Rift head mounted display",
+	"Category" : "Head Mounted Displays.Devices",
+	"EnabledByDefault" : true,
+
+	"Modules" :
+	[
+		{
+			"Name" : "OculusRift",
+			"Type" : "Runtime",
+			"WhitelistPlatforms" : [ "Win64", "Win32", "Mac" ],
+			"LoadingPhase" : "PostConfigInit"
+		}
+	]
+}
diff --git a/Engine/Plugins/Runtime/OculusRift/Source/OculusRift/OculusRift.Build.cs b/Engine/Plugins/Runtime/OculusRift/Source/OculusRift/OculusRift.Build.cs
new file mode 100644
--- /dev/null
+++ b/Engine/Plugins/Runtime/OculusRift/Source/OculusRift/OculusRift.Build.cs
@@ -0,0 +1,49 @@
+// Copyright 1998-2013 Epic Games, Inc. All Rights Reserved.
+
+namespace UnrealBuildTool.Rules
+{
+	public class OculusRift : ModuleRules
+	{
+		public OculusRift(TargetInfo Target)
+		{
+			PrivateIncludePaths.AddRange(
+				new string[] {
+					"OculusRift/Private",
+ 					"../../../../Source/Runtime/Renderer/Private",
+					// ... add other private include paths required here ...
+				}
+				);
+
+			PrivateDependencyModuleNames.AddRange(
+				new string[]
+				{
+					"Core",
+					"CoreUObject",
+					"Engine",
+					"InputCore",
+					"RHI",
+					"RenderCore",
+					"Renderer",
+					"ShaderCore",
+					"HeadMountedDisplay",
+				}
+				);
+
+            // Currently, the Rift is only supported on windows and mac platforms
+            if (Target.Platform == UnrealTargetPlatform.Win32 || Target.Platform == UnrealTargetPlatform.Win64 || Target.Platform == UnrealTargetPlatform.Mac)
+            {
+				PrivateDependencyModuleNames.AddRange(new string[] { "LibOVR", "OpenGLDrv" });
+
+                // Add direct rendering dependencies on a per-platform basis
+                if (Target.Platform == UnrealTargetPlatform.Win32 || Target.Platform == UnrealTargetPlatform.Win64)
+                {
+                    PrivateDependencyModuleNames.AddRange(new string[] { "D3D11RHI" });
+                }
+                if (Target.Platform == UnrealTargetPlatform.Mac)
+                {
+               		AddThirdPartyPrivateStaticDependencies(Target, "OpenGL");
+                }
+            }
+		}
+	}
+}
\ No newline at end of file
diff --git a/Engine/Plugins/Runtime/OculusRift/Source/OculusRift/Private/OculusRiftHMD.cpp b/Engine/Plugins/Runtime/OculusRift/Source/OculusRift/Private/OculusRiftHMD.cpp
new file mode 100644
--- /dev/null
+++ b/Engine/Plugins/Runtime/OculusRift/Source/OculusRift/Private/OculusRiftHMD.cpp
@@ -0,0 +1,2020 @@
+// Copyright 1998-2014 Epic Games, Inc. All Rights Reserved.
+
+#include "OculusRiftPrivate.h"
+#include "OculusRiftHMD.h"
+#include "EngineAnalytics.h"
+#include "Runtime/Analytics/Analytics/Public/Interfaces/IAnalyticsProvider.h"
+
+#if !UE_BUILD_SHIPPING
+// Should be changed to CAPI when available.
+#if PLATFORM_SUPPORTS_PRAGMA_PACK
+#pragma pack (push,8)
+#endif
+#include "../Src/Kernel/OVR_Log.h"
+#if PLATFORM_SUPPORTS_PRAGMA_PACK
+#pragma pack (pop)
+#endif
+#endif // #if !UE_BUILD_SHIPPING
+
+//---------------------------------------------------
+// Oculus Rift Plugin Implementation
+//---------------------------------------------------
+
+class FOculusRiftPlugin : public IOculusRiftPlugin
+{
+	/** IHeadMountedDisplayModule implementation */
+	virtual TSharedPtr< class IHeadMountedDisplay > CreateHeadMountedDisplay() override;
+
+	// Pre-init the HMD module (optional).
+	virtual void PreInit() override;
+};
+
+IMPLEMENT_MODULE( FOculusRiftPlugin, OculusRift )
+
+TSharedPtr< class IHeadMountedDisplay > FOculusRiftPlugin::CreateHeadMountedDisplay()
+{
+#if OCULUS_RIFT_SUPPORTED_PLATFORMS
+	TSharedPtr< FOculusRiftHMD > OculusRiftHMD( new FOculusRiftHMD() );
+	if( OculusRiftHMD->IsInitialized() )
+	{
+		return OculusRiftHMD;
+	}
+#endif//OCULUS_RIFT_SUPPORTED_PLATFORMS
+	return NULL;
+}
+
+void FOculusRiftPlugin::PreInit()
+{
+#if OCULUS_RIFT_SUPPORTED_PLATFORMS
+	FOculusRiftHMD::PreInit();
+#endif//OCULUS_RIFT_SUPPORTED_PLATFORMS
+}
+
+//---------------------------------------------------
+// Oculus Rift IHeadMountedDisplay Implementation
+//---------------------------------------------------
+
+#if OCULUS_RIFT_SUPPORTED_PLATFORMS
+
+#if !UE_BUILD_SHIPPING
+//////////////////////////////////////////////////////////////////////////
+class OculusLog : public OVR::Log
+{
+public:
+	OculusLog()
+	{
+		SetLoggingMask(OVR::LogMask_Debug | OVR::LogMask_Regular);
+	}
+
+	// This virtual function receives all the messages,
+	// developers should override this function in order to do custom logging
+	virtual void    LogMessageVarg(LogMessageType messageType, const char* fmt, va_list argList)
+	{
+		if ((messageType & GetLoggingMask()) == 0)
+			return;
+
+		ANSICHAR buf[1024];
+		int32 len = FCStringAnsi::GetVarArgs(buf, sizeof(buf), sizeof(buf) / sizeof(ANSICHAR), fmt, argList);
+		if (len > 0 && buf[len - 1] == '\n') // truncate the trailing new-line char, since Logf adds its own
+			buf[len - 1] = '\0';
+		TCHAR* tbuf = ANSI_TO_TCHAR(buf);
+		GLog->Logf(TEXT("OCULUS: %s"), tbuf);
+	}
+};
+
+#endif
+
+//////////////////////////////////////////////////////////////////////////
+class ConditionalLocker
+{
+public:
+	ConditionalLocker(bool condition, OVR::Lock* plock) :
+		pLock((condition) ? plock : NULL)
+	{
+		OVR_ASSERT(!condition || pLock);
+		if (pLock)
+			pLock->DoLock();
+	}
+	~ConditionalLocker()
+	{
+		if (pLock)
+			pLock->Unlock();
+	}
+private:
+	OVR::Lock*	pLock;
+};
+
+void FOculusRiftHMD::PreInit()
+{
+	ovr_Initialize();
+}
+
+bool FOculusRiftHMD::IsHMDEnabled() const
+{
+	return bHMDEnabled;
+}
+void FOculusRiftHMD::EnableHMD(bool enable)
+{
+	bHMDEnabled = enable;
+	if (!bHMDEnabled)
+	{
+		EnableStereo(false);
+	}
+}
+
+EHMDDeviceType::Type FOculusRiftHMD::GetHMDDeviceType() const
+{
+	return EHMDDeviceType::DT_OculusRift;
+}
+
+bool FOculusRiftHMD::GetHMDMonitorInfo(MonitorInfo& MonitorDesc) const
+{
+	if (IsInitialized() && Hmd)
+	{
+		MonitorDesc.MonitorName = Hmd->DisplayDeviceName;
+		MonitorDesc.MonitorId	= Hmd->DisplayId;
+		MonitorDesc.DesktopX	= Hmd->WindowsPos.x;
+		MonitorDesc.DesktopY	= Hmd->WindowsPos.y;
+		MonitorDesc.ResolutionX = Hmd->Resolution.w;
+		MonitorDesc.ResolutionY = Hmd->Resolution.h;
+		return true;
+	}
+	else
+	{
+		MonitorDesc.MonitorName = "";
+		MonitorDesc.MonitorId = 0;
+		MonitorDesc.DesktopX = MonitorDesc.DesktopY = MonitorDesc.ResolutionX = MonitorDesc.ResolutionY = 0;
+	}
+	return false;
+}
+
+bool FOculusRiftHMD::IsFullScreenAllowed() const
+{
+	return (Hmd && (Hmd->HmdCaps & ovrHmdCap_ExtendDesktop) != 0) ? true : false;
+}
+
+bool FOculusRiftHMD::DoesSupportPositionalTracking() const
+{
+#ifdef OVR_VISION_ENABLED
+	 return bHmdPosTracking && (SupportedTrackingCaps & ovrTrackingCap_Position) != 0;
+#else
+	return false;
+#endif //OVR_VISION_ENABLED
+}
+
+bool FOculusRiftHMD::HasValidTrackingPosition() const
+{
+#ifdef OVR_VISION_ENABLED
+	return bHmdPosTracking && bHaveVisionTracking;
+#else
+	return false;
+#endif //OVR_VISION_ENABLED
+}
+
+#define CAMERA_HFOV			74.0f
+#define CAMERA_VFOV			54.0f
+#define CAMERA_MIN_DISTANCE 0.25f // meters
+#define CAMERA_MAX_DISTANCE 1.80f // meters
+#define CAMERA_DISTANCE		1.00f // meters (focal point to origin for position)
+
+void FOculusRiftHMD::GetPositionalTrackingCameraProperties(FVector& OutOrigin, FRotator& OutOrientation, float& OutHFOV, float& OutVFOV, float& OutCameraDistance, float& OutNearPlane, float& OutFarPlane) const
+{
+	OutHFOV = CAMERA_HFOV;
+	OutVFOV = CAMERA_VFOV;
+	OutNearPlane  = CAMERA_MIN_DISTANCE * WorldToMetersScale;
+	OutFarPlane   = CAMERA_MAX_DISTANCE * WorldToMetersScale;
+	OutCameraDistance = CAMERA_DISTANCE * WorldToMetersScale;
+
+	// correct position according to BaseOrientation and BaseOffset
+	const FVector off = ToFVector_M2U(BaseOffset);
+
+	const FQuat Orient = BaseOrientation.Inverse() * DeltaControlOrientation;
+	OutOrientation = Orient.Rotator();
+
+	// Calculate origin: where player's eyes are located in the world RELATIVELY to current player's location. This is where we'd
+	// need to translate pre-rotated (using OutOrientation) frustum (again: plus player's current location).
+	const FVector Origin = off - MetersToUU(FVector(CAMERA_DISTANCE, 0, 0)); // the focal point distance
+	OutOrigin = FVector::ZeroVector - Orient.RotateVector(Origin);
+}
+
+bool FOculusRiftHMD::IsInLowPersistenceMode() const
+{
+	return bLowPersistenceMode && (SupportedHmdCaps & ovrHmdCap_LowPersistence) != 0;
+}
+
+void FOculusRiftHMD::EnableLowPersistenceMode(bool Enable)
+{
+	bLowPersistenceMode = Enable;
+	UpdateHmdCaps();
+}
+
+float FOculusRiftHMD::GetInterpupillaryDistance() const
+{
+	return InterpupillaryDistance;
+}
+
+void FOculusRiftHMD::SetInterpupillaryDistance(float NewInterpupillaryDistance)
+{
+	InterpupillaryDistance = NewInterpupillaryDistance;
+
+	UpdateStereoRenderingParams();
+}
+
+void FOculusRiftHMD::GetFieldOfView(float& OutHFOVInDegrees, float& OutVFOVInDegrees) const
+{
+	OutHFOVInDegrees = FMath::RadiansToDegrees(HFOVInRadians);
+	OutVFOVInDegrees = FMath::RadiansToDegrees(VFOVInRadians);
+}
+
+void FOculusRiftHMD::PoseToOrientationAndPosition(const ovrPosef& InPose, FQuat& OutOrientation, FVector& OutPosition) const
+{
+	OutOrientation = ToFQuat(InPose.Orientation);
+
+	// correct position according to BaseOrientation and BaseOffset. Note, if VISION is disabled then BaseOffset is always a zero vector.
+	OutPosition = BaseOrientation.Inverse().RotateVector(ToFVector_M2U(Vector3f(InPose.Position) - BaseOffset));
+
+	// apply base orientation correction to OutOrientation
+	OutOrientation = BaseOrientation.Inverse() * OutOrientation;
+	OutOrientation.Normalize();
+}
+
+void FOculusRiftHMD::GetCurrentOrientationAndPosition(FQuat& CurrentOrientation, FVector& CurrentPosition) const
+{
+	// only supposed to be used from the game thread
+	checkf(IsInGameThread());
+	GetCurrentPose(CurHmdOrientation, CurHmdPosition);
+	CurrentOrientation = LastHmdOrientation = CurHmdOrientation;
+
+	CurrentPosition = CurHmdPosition;
+	
+	// @TODO: we can't actually return just CurHmdPosition here, it should be rotated by CurrentOrientation.
+	// Another problem: if we return rotated position here, we shouldn't apply position in CalculateStereoViewOffset, otherwise the
+	// position will be applied twice.
+	// Currently, position will be applied automatically in CalculateStereoViewOffset and adjusted on Render Thread by UpdatePlayerViewPoint()
+	//CurrentPosition = FVector::ZeroVector; 
+}
+
+void FOculusRiftHMD::GetCurrentPose(FQuat& CurrentOrientation, FVector& CurrentPosition) const
+{
+	const ovrTrackingState ts = ovrHmd_GetTrackingState(Hmd, ovr_GetTimeInSeconds() + MotionPredictionInSeconds);
+	const ovrPosef& ThePose = ts.HeadPose.ThePose;
+	PoseToOrientationAndPosition(ThePose, CurrentOrientation, CurrentPosition);
+	//UE_LOG(LogHMD, Log, TEXT("P: %.3f %.3f %.3f"), CurrentPosition.X, CurrentPosition.Y, CurrentPosition.Y);
+#ifdef OVR_VISION_ENABLED
+	if (bHmdPosTracking)
+	{
+#if !UE_BUILD_SHIPPING
+		bool hadVisionTracking = bHaveVisionTracking;
+		bHaveVisionTracking = (ts.StatusFlags & ovrStatus_PositionTracked) != 0;
+		if (bHaveVisionTracking && !hadVisionTracking)
+			UE_LOG(LogHMD, Warning, TEXT("Vision Tracking Acquired"));
+		if (!bHaveVisionTracking && hadVisionTracking)
+			UE_LOG(LogHMD, Warning, TEXT("Lost Vision Tracking"));
+#endif // #if !UE_BUILD_SHIPPING
+	}
+#endif // OVR_VISION_ENABLED
+}
+
+void FOculusRiftHMD::UpdatePlayerViewPoint(const FQuat& CurrentOrientation, const FVector& CurrentPosition, 
+										   const FVector& LastHmdPosition, const FQuat& DeltaControlOrientation, 
+										   const FQuat& BaseViewOrientation, const FVector& BaseViewPosition, 
+										   FRotator& ViewRotation, FVector& ViewLocation)
+{
+	const FQuat DeltaOrient = BaseViewOrientation.Inverse() * CurrentOrientation;
+	ViewRotation = FRotator(ViewRotation.Quaternion() * DeltaOrient);
+
+	// Apply delta between current HMD position and the LastHmdPosition to ViewLocation
+	const FVector vHMDPositionDelta = DeltaControlOrientation.RotateVector(CurrentPosition - LastHmdPosition);
+	ViewLocation += vHMDPositionDelta;
+}
+
+void FOculusRiftHMD::ApplyHmdRotation(APlayerController* PC, FRotator& ViewRotation)
+{
+#if !UE_BUILD_SHIPPING
+	if (bDoNotUpdateOnGT)
+		return;
+#endif
+	ConditionalLocker lock(bUpdateOnRT, &UpdateOnRTLock);
+
+	// Temporary turn off this method if bUpdateOnRt is true.
+	// 	if (bUpdateOnRT)
+    //	return;
+
+	ViewRotation.Normalize();
+
+	GetCurrentPose(CurHmdOrientation, CurHmdPosition);
+	LastHmdOrientation = CurHmdOrientation;
+
+	const FRotator DeltaRot = ViewRotation - PC->GetControlRotation();
+	DeltaControlRotation = (DeltaControlRotation + DeltaRot).GetNormalized();
+
+	// Pitch from other sources is never good, because there is an absolute up and down that must be respected to avoid motion sickness.
+	// Same with roll.
+	DeltaControlRotation.Pitch = 0;
+	DeltaControlRotation.Roll = 0;
+	DeltaControlOrientation = DeltaControlRotation.Quaternion();
+
+	ViewRotation = FRotator(DeltaControlOrientation * CurHmdOrientation);
+
+#if !UE_BUILD_SHIPPING
+	if (bDrawTrackingCameraFrustum && PC->GetPawnOrSpectator())
+	{
+		DrawDebugTrackingCameraFrustum(PC->GetWorld(), PC->GetPawnOrSpectator()->GetPawnViewLocation());
+	}
+#endif
+}
+
+void FOculusRiftHMD::UpdatePlayerCameraRotation(APlayerCameraManager* Camera, struct FMinimalViewInfo& POV)
+{
+#if !UE_BUILD_SHIPPING
+	if (bDoNotUpdateOnGT)
+		return;
+#endif
+	ConditionalLocker lock(bUpdateOnRT, &UpdateOnRTLock);
+
+	GetCurrentPose(CurHmdOrientation, CurHmdPosition);
+	LastHmdOrientation = CurHmdOrientation;
+
+	static const FName NAME_Fixed = FName(TEXT("Fixed"));
+	static const FName NAME_ThirdPerson = FName(TEXT("ThirdPerson"));
+	static const FName NAME_FreeCam = FName(TEXT("FreeCam"));
+	static const FName NAME_FreeCam_Default = FName(TEXT("FreeCam_Default"));
+	static const FName NAME_FirstPerson = FName(TEXT("FirstPerson"));
+
+	DeltaControlRotation = POV.Rotation;
+	DeltaControlOrientation = DeltaControlRotation.Quaternion();
+
+	// Apply HMD orientation to camera rotation.
+	POV.Rotation = FRotator(POV.Rotation.Quaternion() * CurHmdOrientation);
+
+#if !UE_BUILD_SHIPPING
+	if (bDrawTrackingCameraFrustum)
+	{
+		DrawDebugTrackingCameraFrustum(Camera->GetWorld(), POV.Location);
+	}
+#endif
+}
+
+#if !UE_BUILD_SHIPPING
+void FOculusRiftHMD::DrawDebugTrackingCameraFrustum(UWorld* World, const FVector& ViewLocation)
+{
+	const FColor c = (HasValidTrackingPosition() ? FColor::Green : FColor::Red);
+	FVector origin;
+	FRotator rotation;
+	float hfovDeg, vfovDeg, nearPlane, farPlane, cameraDist;
+	GetPositionalTrackingCameraProperties(origin, rotation, hfovDeg, vfovDeg, cameraDist, nearPlane, farPlane);
+
+	// Level line
+	//DrawDebugLine(World, ViewLocation, FVector(ViewLocation.X + 1000, ViewLocation.Y, ViewLocation.Z), FColor::Blue);
+
+	const float hfov = FMath::DegreesToRadians(hfovDeg * 0.5f);
+	const float vfov = FMath::DegreesToRadians(vfovDeg * 0.5f);
+	FVector coneTop(0, 0, 0);
+	FVector coneBase1(-farPlane, farPlane * FMath::Tan(hfov), farPlane * FMath::Tan(vfov));
+	FVector coneBase2(-farPlane,-farPlane * FMath::Tan(hfov), farPlane * FMath::Tan(vfov));
+	FVector coneBase3(-farPlane,-farPlane * FMath::Tan(hfov),-farPlane * FMath::Tan(vfov));
+	FVector coneBase4(-farPlane, farPlane * FMath::Tan(hfov),-farPlane * FMath::Tan(vfov));
+	FMatrix m(FMatrix::Identity);
+	m *= FRotationMatrix(rotation);
+	m *= FTranslationMatrix(origin);
+	m *= FTranslationMatrix(ViewLocation); // to location of pawn
+	coneTop = m.TransformPosition(coneTop);
+	coneBase1 = m.TransformPosition(coneBase1);
+	coneBase2 = m.TransformPosition(coneBase2);
+	coneBase3 = m.TransformPosition(coneBase3);
+	coneBase4 = m.TransformPosition(coneBase4);
+
+	// draw a point at the camera pos
+	DrawDebugPoint(World, coneTop, 5, c);
+
+	// draw main pyramid, from top to base
+	DrawDebugLine(World, coneTop, coneBase1, c);
+	DrawDebugLine(World, coneTop, coneBase2, c);
+	DrawDebugLine(World, coneTop, coneBase3, c);
+	DrawDebugLine(World, coneTop, coneBase4, c);
+											  
+	// draw base (far plane)				  
+	DrawDebugLine(World, coneBase1, coneBase2, c);
+	DrawDebugLine(World, coneBase2, coneBase3, c);
+	DrawDebugLine(World, coneBase3, coneBase4, c);
+	DrawDebugLine(World, coneBase4, coneBase1, c);
+
+	// draw near plane
+	FVector coneNear1(-nearPlane, nearPlane * FMath::Tan(hfov), nearPlane * FMath::Tan(vfov));
+	FVector coneNear2(-nearPlane,-nearPlane * FMath::Tan(hfov), nearPlane * FMath::Tan(vfov));
+	FVector coneNear3(-nearPlane,-nearPlane * FMath::Tan(hfov),-nearPlane * FMath::Tan(vfov));
+	FVector coneNear4(-nearPlane, nearPlane * FMath::Tan(hfov),-nearPlane * FMath::Tan(vfov));
+	coneNear1 = m.TransformPosition(coneNear1);
+	coneNear2 = m.TransformPosition(coneNear2);
+	coneNear3 = m.TransformPosition(coneNear3);
+	coneNear4 = m.TransformPosition(coneNear4);
+	DrawDebugLine(World, coneNear1, coneNear2, c);
+	DrawDebugLine(World, coneNear2, coneNear3, c);
+	DrawDebugLine(World, coneNear3, coneNear4, c);
+	DrawDebugLine(World, coneNear4, coneNear1, c);
+
+	// center line
+	FVector centerLine(-cameraDist, 0, 0);
+	centerLine = m.TransformPosition(centerLine);
+	DrawDebugLine(World, coneTop, centerLine, FColor::Yellow);
+	DrawDebugPoint(World, centerLine, 5, FColor::Yellow);
+}
+#endif // #if !UE_BUILD_SHIPPING
+
+bool FOculusRiftHMD::IsChromaAbCorrectionEnabled() const
+{
+	return bChromaAbCorrectionEnabled;
+}
+
+ISceneViewExtension* FOculusRiftHMD::GetViewExtension()
+{
+	return this;
+}
+
+bool FOculusRiftHMD::Exec( UWorld* InWorld, const TCHAR* Cmd, FOutputDevice& Ar )
+{
+	if (FParse::Command( &Cmd, TEXT("STEREO") ))
+	{
+		if (FParse::Command(&Cmd, TEXT("ON")))
+		{
+			if (!IsHMDEnabled())
+			{
+				Ar.Logf(TEXT("HMD is disabled. Use 'hmd enable' to re-enable it."));
+			}
+			EnableStereo(true);
+			return true;
+		}
+		else if (FParse::Command(&Cmd, TEXT("OFF")))
+		{
+			EnableStereo(false);
+			return true;
+		}
+		else if (FParse::Command( &Cmd, TEXT("RESET") ))
+		{
+			bOverrideStereo = false;
+			bOverrideIPD = false;
+			bWorldToMetersOverride = false;
+			NearClippingPlane = FarClippingPlane = 0.f;
+			InterpupillaryDistance = ovrHmd_GetFloat(Hmd, OVR_KEY_IPD, OVR_DEFAULT_IPD);
+			UpdateStereoRenderingParams();
+			return true;
+		}
+		else if (FParse::Command(&Cmd, TEXT("SHOW")))
+		{
+			Ar.Logf(TEXT("stereo ipd=%.4f hfov=%.3f vfov=%.3f\n nearPlane=%.4f farPlane=%.4f"), GetInterpupillaryDistance(),
+				FMath::RadiansToDegrees(HFOVInRadians), FMath::RadiansToDegrees(VFOVInRadians),
+				(NearClippingPlane) ? NearClippingPlane : GNearClippingPlane, FarClippingPlane);
+		}
+
+		// normal configuration
+		float val;
+		if (FParse::Value( Cmd, TEXT("E="), val))
+		{
+			SetInterpupillaryDistance( val );
+			bOverrideIPD = true;
+		}
+		if (FParse::Value(Cmd, TEXT("FCP="), val)) // far clipping plane override
+		{
+			FarClippingPlane = val;
+		}
+		if (FParse::Value(Cmd, TEXT("NCP="), val)) // near clipping plane override
+		{
+			NearClippingPlane = val;
+		}
+		if (FParse::Value(Cmd, TEXT("W2M="), val))
+		{
+			WorldToMetersScale = val;
+			bWorldToMetersOverride = true;
+		}
+
+		// debug configuration
+		if (bDevSettingsEnabled)
+		{
+			float fov;
+			if (FParse::Value(Cmd, TEXT("HFOV="), fov))
+			{
+				HFOVInRadians = FMath::DegreesToRadians(fov);
+				bOverrideStereo = true;
+			}
+			else if (FParse::Value(Cmd, TEXT("VFOV="), fov))
+			{
+				VFOVInRadians = FMath::DegreesToRadians(fov);
+				bOverrideStereo = true;
+			}
+		}
+		return true;
+	}
+	else if (FParse::Command(&Cmd, TEXT("HMD")))
+	{
+		if (FParse::Command(&Cmd, TEXT("ENABLE")))
+		{
+			EnableHMD(true);
+			return true;
+		}
+		else if (FParse::Command(&Cmd, TEXT("DISABLE")))
+		{
+			EnableHMD(false);
+			return true;
+		}
+		else if (FParse::Command(&Cmd, TEXT("VSYNC")))
+		{
+			if (FParse::Command( &Cmd, TEXT("RESET") ))
+			{
+				if (bStereoEnabled)
+				{
+					bVSync = bSavedVSync;
+					ApplySystemOverridesOnStereo();
+				}
+				bOverrideVSync = false;
+				return true;
+			}
+			else
+			{
+				if (FParse::Command(&Cmd, TEXT("ON")) || FParse::Command(&Cmd, TEXT("1")))
+				{
+					bVSync = true;
+					bOverrideVSync = true;
+					ApplySystemOverridesOnStereo();
+					return true;
+				}
+				else if (FParse::Command(&Cmd, TEXT("OFF")) || FParse::Command(&Cmd, TEXT("0")))
+				{
+					bVSync = false;
+					bOverrideVSync = true;
+					ApplySystemOverridesOnStereo();
+					return true;
+				}
+				else if (FParse::Command(&Cmd, TEXT("TOGGLE")) || FParse::Command(&Cmd, TEXT("")))
+				{
+					bVSync = !bVSync;
+					bOverrideVSync = true;
+					ApplySystemOverridesOnStereo();
+					Ar.Logf(TEXT("VSync is currently %s"), (bVSync) ? TEXT("ON") : TEXT("OFF"));
+					return true;
+				}
+			}
+			return false;
+		}
+		else if (FParse::Command(&Cmd, TEXT("SP")) || 
+				 FParse::Command(&Cmd, TEXT("SCREENPERCENTAGE")))
+		{
+			FString CmdName = FParse::Token(Cmd, 0);
+			if (CmdName.IsEmpty())
+				return false;
+
+			if (!FCString::Stricmp(*CmdName, TEXT("RESET")))
+			{
+				bOverrideScreenPercentage = false;
+				ApplySystemOverridesOnStereo();
+			}
+			else
+			{
+				float sp = FCString::Atof(*CmdName);
+				if (sp >= 30 && sp <= 300)
+				{
+					bOverrideScreenPercentage = true;
+					ScreenPercentage = sp;
+					ApplySystemOverridesOnStereo();
+				}
+				else
+				{
+					Ar.Logf(TEXT("Value is out of range [30..300]"));
+				}
+			}
+			return true;
+		}
+		else if (FParse::Command(&Cmd, TEXT("LP"))) // low persistence mode
+		{
+			FString CmdName = FParse::Token(Cmd, 0);
+			if (!CmdName.IsEmpty())
+			{
+				if (!FCString::Stricmp(*CmdName, TEXT("ON")))
+				{
+					bLowPersistenceMode = true;
+				}
+				else if (!FCString::Stricmp(*CmdName, TEXT("OFF")))
+				{
+					bLowPersistenceMode = false;
+				}
+				else if (!FCString::Stricmp(*CmdName, TEXT("TOGGLE")))
+				{
+					bLowPersistenceMode = !bLowPersistenceMode;
+				}
+				else
+				{
+					return false;
+				}
+			}
+			else
+			{
+				bLowPersistenceMode = !bLowPersistenceMode;
+			}
+			UpdateHmdCaps();
+			Ar.Logf(TEXT("Low Persistence is currently %s"), (bLowPersistenceMode) ? TEXT("ON") : TEXT("OFF"));
+			return true;
+		}
+		else if (FParse::Command(&Cmd, TEXT("MIRROR"))) // to mirror or not to mirror?...
+		{
+			FString CmdName = FParse::Token(Cmd, 0);
+			if (!CmdName.IsEmpty())
+			{
+				if (!FCString::Stricmp(*CmdName, TEXT("ON")))
+				{
+					bMirrorToWindow = true;
+				}
+				else if (!FCString::Stricmp(*CmdName, TEXT("OFF")))
+				{
+					bMirrorToWindow = false;
+				}
+				else if (!FCString::Stricmp(*CmdName, TEXT("TOGGLE"))) 
+				{
+					bMirrorToWindow = !bMirrorToWindow;
+				}
+			}
+			else
+			{
+				bMirrorToWindow = !bMirrorToWindow;
+			}
+			UpdateHmdCaps();
+			Ar.Logf(TEXT("Mirroring is currently %s"), (bMirrorToWindow) ? TEXT("ON") : TEXT("OFF"));
+			return true;
+		}
+		else if (FParse::Command(&Cmd, TEXT("UPDATEONRT"))) // update on renderthread
+		{
+			FString CmdName = FParse::Token(Cmd, 0);
+			if (!CmdName.IsEmpty())
+			{
+				if (!FCString::Stricmp(*CmdName, TEXT("ON")))
+				{
+					bUpdateOnRT = true;
+				}
+				else if (!FCString::Stricmp(*CmdName, TEXT("OFF")))
+				{
+					bUpdateOnRT = false;
+				}
+				else if (!FCString::Stricmp(*CmdName, TEXT("TOGGLE")))
+				{
+					bUpdateOnRT = !bUpdateOnRT;
+				}
+				else
+				{
+					return false;
+				}
+			}
+			else
+			{
+				bUpdateOnRT = !bUpdateOnRT;
+			}
+			Ar.Logf(TEXT("Update on render thread is currently %s"), (bUpdateOnRT) ? TEXT("ON") : TEXT("OFF"));
+			return true;
+		}
+		else if (FParse::Command(&Cmd, TEXT("OVERDRIVE"))) // 2 frame raise overdrive
+		{
+			FString CmdName = FParse::Token(Cmd, 0);
+			if (!CmdName.IsEmpty())
+			{
+				if (!FCString::Stricmp(*CmdName, TEXT("ON")))
+				{
+					bOverdrive = true;
+				}
+				else if (!FCString::Stricmp(*CmdName, TEXT("OFF")))
+				{
+					bOverdrive = false;
+				}
+				else if (!FCString::Stricmp(*CmdName, TEXT("TOGGLE")))
+				{
+					bOverdrive = !bOverdrive;
+				}
+				else
+				{
+					return false;
+				}
+			}
+			else
+			{
+				bOverdrive = !bOverdrive;
+			}
+			UpdateDistortionCaps();
+			Ar.Logf(TEXT("Overdrive is currently %s"), (bOverdrive) ? TEXT("ON") : TEXT("OFF"));
+			return true;
+		}
+#ifdef OVR_DIRECT_RENDERING
+		else if (FParse::Command(&Cmd, TEXT("TIMEWARP"))) 
+		{
+			FString CmdName = FParse::Token(Cmd, 0);
+			if (!CmdName.IsEmpty())
+			{
+				if (!FCString::Stricmp(*CmdName, TEXT("ON")))
+				{
+					bTimeWarp = true;
+				}
+				else if (!FCString::Stricmp(*CmdName, TEXT("OFF")))
+				{
+					bTimeWarp = false;
+				}
+				else if (!FCString::Stricmp(*CmdName, TEXT("TOGGLE")))
+				{
+					bTimeWarp = !bTimeWarp;
+				}
+				else
+				{
+					return false;
+				}
+			}
+			else
+			{
+				bTimeWarp = !bTimeWarp;
+			}
+			UpdateDistortionCaps();
+			Ar.Logf(TEXT("TimeWarp is currently %s"), (bTimeWarp) ? TEXT("ON") : TEXT("OFF"));
+			return true;
+		}
+#endif // #ifdef OVR_DIRECT_RENDERING
+#if !UE_BUILD_SHIPPING
+		else if (FParse::Command(&Cmd, TEXT("UPDATEONGT"))) // update on game thread
+		{
+			FString CmdName = FParse::Token(Cmd, 0);
+			if (!CmdName.IsEmpty())
+			{
+				if (!FCString::Stricmp(*CmdName, TEXT("ON")))
+				{
+					bDoNotUpdateOnGT = false;
+				}
+				else if (!FCString::Stricmp(*CmdName, TEXT("OFF")))
+				{
+					bDoNotUpdateOnGT = true;
+				}
+				else if (!FCString::Stricmp(*CmdName, TEXT("TOGGLE")))
+				{
+					bDoNotUpdateOnGT = !bDoNotUpdateOnGT;
+				}
+				else
+				{
+					return false;
+				}
+			}
+			else
+			{
+				bDoNotUpdateOnGT = !bDoNotUpdateOnGT;
+			}
+			Ar.Logf(TEXT("Update on game thread is currently %s"), (!bDoNotUpdateOnGT) ? TEXT("ON") : TEXT("OFF"));
+			return true;
+		}
+		else if (FParse::Command(&Cmd, TEXT("STATS"))) // status / statistics
+		{
+			bShowStats = !bShowStats;
+			return true;
+		}
+		else if (FParse::Command(&Cmd, TEXT("GRID"))) // grid
+		{
+			bDrawGrid = !bDrawGrid;
+			return true;
+		}
+		else if (FParse::Command(&Cmd, TEXT("PROFILE"))) // profile
+		{
+			bProfiling = !bProfiling;
+			UpdateDistortionCaps();
+			Ar.Logf(TEXT("Profiling mode is currently %s"), (bProfiling) ? TEXT("ON") : TEXT("OFF"));
+			return true;
+		}
+#endif //UE_BUILD_SHIPPING
+	}
+	else if (FParse::Command(&Cmd, TEXT("HMDMAG")))
+	{
+		if (FParse::Command(&Cmd, TEXT("ON")))
+		{
+			bYawDriftCorrectionEnabled = true;
+			UpdateHmdCaps();
+			return true;
+		}
+		else if (FParse::Command(&Cmd, TEXT("OFF")))
+		{
+			bYawDriftCorrectionEnabled = false;
+			UpdateHmdCaps();
+			return true;
+		}
+		else if (FParse::Command(&Cmd, TEXT("SHOW")))
+		{
+			Ar.Logf(TEXT("mag %s"), bYawDriftCorrectionEnabled ? 
+				TEXT("on") : TEXT("off"));
+			return true;
+		}
+		return false;
+	}
+	else if (FParse::Command(&Cmd, TEXT("HMDWARP")))
+    {
+#ifndef OVR_DIRECT_RENDERING
+        if (FParse::Command( &Cmd, TEXT("ON") ))
+        {
+            bHmdDistortion = true;
+            return true;
+        }
+        else if (FParse::Command( &Cmd, TEXT("OFF") ))
+        {
+            bHmdDistortion = false;
+            return true;
+        }
+#endif //OVR_DIRECT_RENDERING
+		if (FParse::Command(&Cmd, TEXT("CHA")))
+		{
+			bChromaAbCorrectionEnabled = true;
+			UpdateDistortionCaps();
+			return true;
+		}
+		else if (FParse::Command(&Cmd, TEXT("NOCHA")))
+		{
+			bChromaAbCorrectionEnabled = false;
+			UpdateDistortionCaps();
+			return true;
+		}
+		else if (FParse::Command( &Cmd, TEXT("HQ") ))
+		{
+			// High quality distortion
+			if (FParse::Command( &Cmd, TEXT("ON") ))
+			{
+				bHQDistortion = true;
+			}
+			else if (FParse::Command(&Cmd, TEXT("OFF")))
+			{
+				bHQDistortion = false;
+			}
+			else
+			{
+				bHQDistortion = !bHQDistortion;
+			}
+			Ar.Logf(TEXT("High quality distortion is currently %s"), (bHQDistortion) ? TEXT("ON") : TEXT("OFF") ?
+				TEXT("on") : TEXT("off"));
+			UpdateDistortionCaps();
+			return true;
+		}
+
+		if (FParse::Command(&Cmd, TEXT("SHOW")))
+		{
+			Ar.Logf(TEXT("hmdwarp %s sc=%f %s"), (bHmdDistortion ? TEXT("on") : TEXT("off"))
+				, IdealScreenPercentage / 100.f
+				, (bChromaAbCorrectionEnabled ? TEXT("cha") : TEXT("nocha")));
+		}
+		return true;
+    }
+	else if (FParse::Command(&Cmd, TEXT("HMDPOS")))
+	{
+		if (FParse::Command(&Cmd, TEXT("RESET")))
+		{
+			FString YawStr = FParse::Token(Cmd, 0);
+			float yaw = 0.f;
+			if (!YawStr.IsEmpty())
+			{
+				yaw = FCString::Atof(*YawStr);
+			}
+			ResetOrientationAndPosition(yaw);
+			return true;
+		}
+#ifdef OVR_VISION_ENABLED
+		else if (FParse::Command(&Cmd, TEXT("ON")) || FParse::Command(&Cmd, TEXT("ENABLE")))
+		{
+			bHmdPosTracking = true;
+			UpdateHmdCaps();
+			return true;
+		}
+		else if (FParse::Command(&Cmd, TEXT("OFF")) || FParse::Command(&Cmd, TEXT("DISABLE")))
+		{
+			bHmdPosTracking = false;
+			bHaveVisionTracking = false;
+			UpdateHmdCaps();
+			return true;
+		}
+		else if (FParse::Command(&Cmd, TEXT("TOGGLE")))
+		{
+			bHmdPosTracking = !bHmdPosTracking;
+			bHaveVisionTracking = false;
+			UpdateHmdCaps();
+			return true;
+		}
+#if !UE_BUILD_SHIPPING
+		else if (FParse::Command(&Cmd, TEXT("SHOWCAMERA")))
+		{
+			if (FParse::Command(&Cmd, TEXT("OFF")))
+			{
+				bDrawTrackingCameraFrustum = false;
+				return true;
+			}
+			if (FParse::Command(&Cmd, TEXT("ON")))
+			{
+				bDrawTrackingCameraFrustum = true;
+				return true;
+			}
+			else 
+			{
+				bDrawTrackingCameraFrustum = !bDrawTrackingCameraFrustum;
+				return true;
+			}
+		}
+#endif // #if !UE_BUILD_SHIPPING
+		else if (FParse::Command(&Cmd, TEXT("SHOW")))
+		{
+			Ar.Logf(TEXT("hmdpos is %s, vision='%s'"), 
+				(bHmdPosTracking ? TEXT("enabled") : TEXT("disabled")),
+				(bHaveVisionTracking ? TEXT("active") : TEXT("lost")));
+			return true;
+		}
+#endif // #ifdef OVR_VISION_ENABLED
+	}
+	else if (FParse::Command(&Cmd, TEXT("OCULUSDEV")))
+	{
+		if (FParse::Command(&Cmd, TEXT("ON")))
+		{
+			bDevSettingsEnabled = true;
+		}
+		else if (FParse::Command(&Cmd, TEXT("OFF")))
+		{
+			bDevSettingsEnabled = false;
+		}
+		UpdateStereoRenderingParams();
+		return true;
+	}
+	if (FParse::Command(&Cmd, TEXT("MOTION")))
+	{
+		FString CmdName = FParse::Token(Cmd, 0);
+		if (CmdName.IsEmpty())
+			return false;
+
+		if (!FCString::Stricmp(*CmdName, TEXT("ON")))
+		{
+			bHeadTrackingEnforced = false;
+			return true;
+		}
+		else if (!FCString::Stricmp(*CmdName, TEXT("ENFORCE")))
+		{
+			bHeadTrackingEnforced = !bHeadTrackingEnforced;
+			if (!bHeadTrackingEnforced)
+			{
+				CurHmdOrientation = FQuat::Identity;
+				ResetControlRotation();
+			}
+			return true;
+		}
+		else if (!FCString::Stricmp(*CmdName, TEXT("RESET")))
+		{
+			bHeadTrackingEnforced = false;
+			CurHmdOrientation = FQuat::Identity;
+			ResetControlRotation();
+			return true;
+		}
+		else if (!FCString::Stricmp(*CmdName, TEXT("SHOW")))
+		{
+			if (MotionPredictionInSeconds > 0)
+			{
+				Ar.Logf(TEXT("motion prediction=%.3f"),	MotionPredictionInSeconds);
+			}
+			else
+			{
+				Ar.Logf(TEXT("motion prediction OFF"));
+			}
+	
+			return true;
+		}
+
+		FString Value = FParse::Token(Cmd, 0);
+		if (Value.IsEmpty())
+		{
+			return false;
+		}
+		if (!FCString::Stricmp(*CmdName, TEXT("PRED")))
+		{
+			if (!FCString::Stricmp(*Value, TEXT("OFF")))
+			{
+				MotionPredictionInSeconds = 0.0;
+			}
+			else
+			{
+				MotionPredictionInSeconds = FCString::Atod(*Value);
+			}
+			return true;
+		}
+		return false;
+	}
+#ifndef OVR_DIRECT_RENDERING
+	else if (FParse::Command(&Cmd, TEXT("SETFINISHFRAME")))
+	{
+		static IConsoleVariable* CFinishFrameVar = IConsoleManager::Get().FindConsoleVariable(TEXT("r.FinishCurrentFrame"));
+
+		if (FParse::Command(&Cmd, TEXT("ON")))
+		{
+			bAllowFinishCurrentFrame = true;
+			if (bStereoEnabled)
+			{
+				CFinishFrameVar->Set(bAllowFinishCurrentFrame);
+			}
+			return true;
+		}
+		else if (FParse::Command(&Cmd, TEXT("OFF")))
+		{
+			bAllowFinishCurrentFrame = false;
+			if (bStereoEnabled)
+			{
+				CFinishFrameVar->Set(bAllowFinishCurrentFrame);
+			}
+			return true;
+		}
+		return false;
+	}
+#endif
+	else if (FParse::Command(&Cmd, TEXT("UNCAPFPS")))
+	{
+		GEngine->bSmoothFrameRate = false;
+		return true;
+	}
+	else if (FParse::Command(&Cmd, TEXT("OVRVERSION")))
+	{
+		static const char* Results = ovr_GetVersionString() + 7; // +7 is temporary, until libovr returns just numbers
+		Ar.Logf(TEXT("%s, LibOVR: %s, built %s, %s"), *GEngineVersion.ToString(), UTF8_TO_TCHAR(Results), 
+			UTF8_TO_TCHAR(__DATE__), UTF8_TO_TCHAR(__TIME__));
+		return true;
+	}
+
+	return false;
+}
+
+void FOculusRiftHMD::OnScreenModeChange(EWindowMode::Type WindowMode)
+{
+	EnableStereo(WindowMode != EWindowMode::Windowed);
+	UpdateStereoRenderingParams();
+}
+
+void FOculusRiftHMD::RecordAnalytics()
+{
+	if (FEngineAnalytics::IsAvailable())
+	{
+		// prepare and send analytics data
+		TArray<FAnalyticsEventAttribute> EventAttributes;
+
+		IHeadMountedDisplay::MonitorInfo MonitorInfo;
+		GetHMDMonitorInfo(MonitorInfo);
+		if (Hmd)
+		{
+			EventAttributes.Add(FAnalyticsEventAttribute(TEXT("DeviceName"), FString::Printf(TEXT("%s - %s"), ANSI_TO_TCHAR(Hmd->Manufacturer), ANSI_TO_TCHAR(Hmd->ProductName))));
+		}
+		EventAttributes.Add(FAnalyticsEventAttribute(TEXT("DisplayDeviceName"), MonitorInfo.MonitorName));
+#if PLATFORM_MAC // On OS X MonitorId is the CGDirectDisplayID aka uint64, not a string
+		FString DisplayId(FString::Printf(TEXT("%llu"), MonitorInfo.MonitorId));
+		EventAttributes.Add(FAnalyticsEventAttribute(TEXT("DisplayId"), DisplayId));
+#else
+		EventAttributes.Add(FAnalyticsEventAttribute(TEXT("DisplayId"), MonitorInfo.MonitorId));
+#endif
+		FString MonResolution(FString::Printf(TEXT("(%d, %d)"), MonitorInfo.ResolutionX, MonitorInfo.ResolutionY));
+		EventAttributes.Add(FAnalyticsEventAttribute(TEXT("Resolution"), MonResolution));
+
+		EventAttributes.Add(FAnalyticsEventAttribute(TEXT("ChromaAbCorrectionEnabled"), bChromaAbCorrectionEnabled));
+		EventAttributes.Add(FAnalyticsEventAttribute(TEXT("MagEnabled"), bYawDriftCorrectionEnabled));
+		EventAttributes.Add(FAnalyticsEventAttribute(TEXT("DevSettingsEnabled"), bDevSettingsEnabled));
+		EventAttributes.Add(FAnalyticsEventAttribute(TEXT("MotionPredictionEnabled"), (MotionPredictionInSeconds > 0)));
+		EventAttributes.Add(FAnalyticsEventAttribute(TEXT("AccelGain"), AccelGain));
+		EventAttributes.Add(FAnalyticsEventAttribute(TEXT("OverrideInterpupillaryDistance"), bOverrideIPD));
+		if (bOverrideIPD)
+		{
+			EventAttributes.Add(FAnalyticsEventAttribute(TEXT("InterpupillaryDistance"), GetInterpupillaryDistance()));
+		}
+		EventAttributes.Add(FAnalyticsEventAttribute(TEXT("OverrideStereo"), bOverrideStereo));
+		if (bOverrideStereo)
+		{
+			EventAttributes.Add(FAnalyticsEventAttribute(TEXT("HFOV"), HFOVInRadians));
+			EventAttributes.Add(FAnalyticsEventAttribute(TEXT("VFOV"), VFOVInRadians));
+		}
+		EventAttributes.Add(FAnalyticsEventAttribute(TEXT("OverrideVSync"), bOverrideVSync));
+		if (bOverrideVSync)
+		{
+			EventAttributes.Add(FAnalyticsEventAttribute(TEXT("VSync"), bVSync));
+		}
+		EventAttributes.Add(FAnalyticsEventAttribute(TEXT("OverrideScreenPercentage"), bOverrideScreenPercentage));
+		if (bOverrideScreenPercentage)
+		{
+			EventAttributes.Add(FAnalyticsEventAttribute(TEXT("ScreenPercentage"), ScreenPercentage));
+		}
+		if (bWorldToMetersOverride)
+		{
+			EventAttributes.Add(FAnalyticsEventAttribute(TEXT("WorldToMetersScale"), WorldToMetersScale));
+		}
+		EventAttributes.Add(FAnalyticsEventAttribute(TEXT("InterpupillaryDistance"), InterpupillaryDistance));
+		EventAttributes.Add(FAnalyticsEventAttribute(TEXT("TimeWarp"), bTimeWarp));
+		EventAttributes.Add(FAnalyticsEventAttribute(TEXT("AllowFinishCurrentFrame"), bAllowFinishCurrentFrame));
+#ifdef OVR_VISION_ENABLED
+		EventAttributes.Add(FAnalyticsEventAttribute(TEXT("HmdPosTracking"), bHmdPosTracking));
+		EventAttributes.Add(FAnalyticsEventAttribute(TEXT("LowPersistenceMode"), bLowPersistenceMode));
+#endif		
+		EventAttributes.Add(FAnalyticsEventAttribute(TEXT("UpdateOnRT"), bUpdateOnRT));
+		EventAttributes.Add(FAnalyticsEventAttribute(TEXT("Overdrive"), bOverdrive));
+		EventAttributes.Add(FAnalyticsEventAttribute(TEXT("MirrorToWindow"), bMirrorToWindow));
+
+		FString OutStr(TEXT("Editor.VR.DeviceInitialised"));
+		FEngineAnalytics::GetProvider().RecordEvent(OutStr, EventAttributes);
+	}
+}
+
+bool FOculusRiftHMD::IsPositionalTrackingEnabled() const
+{
+#ifdef OVR_VISION_ENABLED
+	return bHmdPosTracking;
+#else
+	return false;
+#endif // #ifdef OVR_VISION_ENABLED
+}
+
+bool FOculusRiftHMD::EnablePositionalTracking(bool enable)
+{
+#ifdef OVR_VISION_ENABLED
+	bHmdPosTracking = enable;
+	return IsPositionalTrackingEnabled();
+#else
+	OVR_UNUSED(enable);
+	return false;
+#endif
+}
+
+//---------------------------------------------------
+// Oculus Rift IStereoRendering Implementation
+//---------------------------------------------------
+bool FOculusRiftHMD::IsStereoEnabled() const
+{
+	return bStereoEnabled && bHMDEnabled;
+}
+
+bool FOculusRiftHMD::EnableStereo(bool stereo)
+{
+	bStereoEnabled = (IsHMDEnabled()) ? stereo : false;
+	OnOculusStateChange(bStereoEnabled);
+	return bStereoEnabled;
+}
+
+void FOculusRiftHMD::ResetControlRotation() const
+{
+	// Switching back to non-stereo mode: reset player rotation and aim.
+	// Should we go through all playercontrollers here?
+	APlayerController* pc = GEngine->GetFirstLocalPlayerController(GWorld);
+	if (pc)
+	{
+		// Reset Aim? @todo
+		FRotator r = pc->GetControlRotation();
+		r.Normalize();
+		// Reset roll and pitch of the player
+		r.Roll = 0;
+		r.Pitch = 0;
+		pc->SetControlRotation(r);
+	}
+}
+
+void FOculusRiftHMD::OnOculusStateChange(bool bIsEnabledNow)
+{
+	bHmdDistortion = bIsEnabledNow;
+	if (!bIsEnabledNow)
+	{
+		// Switching from stereo
+
+		ResetControlRotation();
+		RestoreSystemValues();
+	}
+	else
+	{
+		SaveSystemValues();
+		ApplySystemOverridesOnStereo(bIsEnabledNow);
+
+		UpdateStereoRenderingParams();
+	}
+	// need to distribute the event to user's code somehow... (!AB) @todo
+}
+
+void FOculusRiftHMD::ApplySystemOverridesOnStereo(bool bForce)
+{
+	if (bStereoEnabled || bForce)
+	{
+		// Set the current VSync state
+		if (bOverrideVSync)
+		{
+			static IConsoleVariable* CVSyncVar = IConsoleManager::Get().FindConsoleVariable(TEXT("r.VSync"));
+			CVSyncVar->Set(bVSync);
+		}
+		else
+		{
+			static IConsoleVariable* CVSyncVar = IConsoleManager::Get().FindConsoleVariable(TEXT("r.VSync"));
+			bVSync = CVSyncVar->GetInt() != 0;
+		}
+		UpdateHmdCaps();
+
+#ifndef OVR_DIRECT_RENDERING
+		static IConsoleVariable* CFinishFrameVar = IConsoleManager::Get().FindConsoleVariable(TEXT("r.FinishCurrentFrame"));
+		CFinishFrameVar->Set(bAllowFinishCurrentFrame);
+#endif
+	}
+}
+
+void FOculusRiftHMD::SaveSystemValues()
+{
+	static IConsoleVariable* CVSyncVar = IConsoleManager::Get().FindConsoleVariable(TEXT("r.VSync"));
+	bSavedVSync = CVSyncVar->GetInt() != 0;
+
+	static IConsoleVariable* CScrPercVar = IConsoleManager::Get().FindConsoleVariable(TEXT("r.ScreenPercentage"));
+	SavedScrPerc = CScrPercVar->GetFloat();
+}
+
+void FOculusRiftHMD::RestoreSystemValues()
+{
+	static IConsoleVariable* CVSyncVar = IConsoleManager::Get().FindConsoleVariable(TEXT("r.VSync"));
+	CVSyncVar->Set(bSavedVSync);
+
+	static IConsoleVariable* CScrPercVar = IConsoleManager::Get().FindConsoleVariable(TEXT("r.ScreenPercentage"));
+	CScrPercVar->Set(SavedScrPerc);
+
+	static IConsoleVariable* CFinishFrameVar = IConsoleManager::Get().FindConsoleVariable(TEXT("r.FinishCurrentFrame"));
+	CFinishFrameVar->Set(false);
+}
+
+void FOculusRiftHMD::UpdateScreenSettings(const FViewport*)
+{
+	// Set the current ScreenPercentage state
+	static IConsoleVariable* CScrPercVar = IConsoleManager::Get().FindConsoleVariable(TEXT("r.ScreenPercentage"));
+	float DesiredSceeenPercentage;
+	if (bOverrideScreenPercentage)
+	{
+		DesiredSceeenPercentage = ScreenPercentage;
+	}
+	else
+	{
+		DesiredSceeenPercentage = IdealScreenPercentage;
+	}
+	if (FMath::RoundToInt(CScrPercVar->GetFloat()) != FMath::RoundToInt(DesiredSceeenPercentage))
+	{
+		CScrPercVar->Set(DesiredSceeenPercentage);
+	}
+}
+
+void FOculusRiftHMD::AdjustViewRect(EStereoscopicPass StereoPass, int32& X, int32& Y, uint32& SizeX, uint32& SizeY) const
+{
+    SizeX = SizeX / 2;
+    if( StereoPass == eSSP_RIGHT_EYE )
+    {
+        X += SizeX;
+    }
+}
+
+void FOculusRiftHMD::CalculateStereoViewOffset(const EStereoscopicPass StereoPassType, const FRotator& ViewRotation, const float WorldToMeters, FVector& ViewLocation)
+{
+	ConditionalLocker lock(bUpdateOnRT, &UpdateOnRTLock);
+	if (bNeedUpdateStereoRenderingParams)
+		UpdateStereoRenderingParams();
+
+	if( StereoPassType != eSSP_FULL )
+	{
+		check(WorldToMeters != 0.f)
+
+		const int idx = (StereoPassType == eSSP_LEFT_EYE) ? 0 : 1;
+        const float PassEyeOffset = -EyeRenderDesc[idx].ViewAdjust.x * WorldToMeters;
+
+		const FVector TotalOffset = FVector(0, PassEyeOffset, 0);
+
+		ViewLocation += ViewRotation.Quaternion().RotateVector(TotalOffset);
+
+		// The HMDPosition already has HMD orientation applied.
+		// Apply rotational difference between HMD orientation and ViewRotation
+		// to HMDPosition vector. 
+
+		const FVector vHMDPosition = DeltaControlOrientation.RotateVector(CurHmdPosition);
+		ViewLocation += vHMDPosition;
+		LastHmdPosition = CurHmdPosition;
+	}
+	else if (bHeadTrackingEnforced)
+	{
+		const FVector vHMDPosition = DeltaControlOrientation.RotateVector(CurHmdPosition);
+		ViewLocation += vHMDPosition;
+		LastHmdPosition = CurHmdPosition;
+	}
+}
+
+void FOculusRiftHMD::ResetOrientationAndPosition(float yaw)
+{
+	const ovrTrackingState ss = ovrHmd_GetTrackingState(Hmd, ovr_GetTimeInSeconds());
+	const ovrPosef& pose = ss.HeadPose.ThePose;
+	const OVR::Quatf orientation = OVR::Quatf(pose.Orientation);
+
+	// Reset position
+#ifdef OVR_VISION_ENABLED
+	BaseOffset = pose.Position;
+#else
+	BaseOffset = OVR::Vector3f(0, 0, 0);
+#endif // #ifdef OVR_VISION_ENABLED
+
+	FRotator ViewRotation;
+	ViewRotation = FRotator(ToFQuat(orientation));
+	ViewRotation.Pitch = 0;
+	ViewRotation.Roll = 0;
+
+	if (yaw != 0.f)
+	{
+		// apply optional yaw offset
+		ViewRotation.Yaw -= yaw;
+		ViewRotation.Normalize();
+	}
+
+	BaseOrientation = ViewRotation.Quaternion();
+}
+
+FMatrix FOculusRiftHMD::GetStereoProjectionMatrix(enum EStereoscopicPass StereoPassType, const float FOV) const
+{
+	// Stereo params must be recalculated already, see CalculateStereoViewOffset.
+	check(!bNeedUpdateStereoRenderingParams);
+	check(IsStereoEnabled());
+
+	const int idx = (StereoPassType == eSSP_LEFT_EYE) ? 0 : 1;
+
+	FMatrix proj = ToFMatrix(EyeProjectionMatrices[idx]);
+
+	// correct far and near planes for reversed-Z projection matrix
+	const float InNearZ = (NearClippingPlane) ? NearClippingPlane : GNearClippingPlane;
+	const float InFarZ = (FarClippingPlane) ? FarClippingPlane : GNearClippingPlane;
+	proj.M[3][3] = 0.0f;
+	proj.M[2][3] = 1.0f;
+
+	proj.M[2][2] = (InNearZ == InFarZ) ? 0.0f    : InNearZ / (InNearZ - InFarZ);
+	proj.M[3][2] = (InNearZ == InFarZ) ? InNearZ : -InFarZ * InNearZ / (InNearZ - InFarZ);
+
+	return proj;
+}
+
+void FOculusRiftHMD::InitCanvasFromView(FSceneView* InView, UCanvas* Canvas)
+{
+	// This is used for placing small HUDs (with names)
+	// over other players (for example, in Capture Flag).
+	// HmdOrientation should be initialized by GetCurrentOrientation (or
+	// user's own value).
+	FSceneView HmdView(*InView);
+
+	UpdatePlayerViewPoint(Canvas->HmdOrientation, FVector(0.f), FVector::ZeroVector, FQuat::Identity, HmdView.BaseHmdOrientation, HmdView.BaseHmdLocation, HmdView.ViewRotation, HmdView.ViewLocation);
+
+	HmdView.UpdateViewMatrix();
+	Canvas->ViewProjectionMatrix = HmdView.ViewProjectionMatrix;
+}
+
+void FOculusRiftHMD::PushViewportCanvas(EStereoscopicPass StereoPass, FCanvas *InCanvas, UCanvas *InCanvasObject, FViewport *InViewport) const
+{
+	if (StereoPass != eSSP_FULL)
+	{
+		int32 SideSizeX = FMath::TruncToInt(InViewport->GetSizeXY().X * 0.5);
+
+		// !AB: temporarily assuming all canvases are at Z = 1.0f and calculating
+		// stereo disparity right here. Stereo disparity should be calculated for each
+		// element separately, considering its actual Z-depth.
+		const float Z = 1.0f;
+		float Disparity = Z * HudOffset + Z * CanvasCenterOffset;
+		if (StereoPass == eSSP_RIGHT_EYE)
+			Disparity = -Disparity;
+
+		if (InCanvasObject)
+		{
+			//InCanvasObject->Init();
+			InCanvasObject->SizeX = SideSizeX;
+			InCanvasObject->SizeY = InViewport->GetSizeXY().Y;
+			InCanvasObject->SetView(NULL);
+			InCanvasObject->Update();
+		}
+
+		float ScaleFactor = 1.0f;
+		FScaleMatrix m(ScaleFactor);
+
+		InCanvas->PushAbsoluteTransform(FTranslationMatrix(
+			FVector(((StereoPass == eSSP_RIGHT_EYE) ? SideSizeX : 0) + Disparity, 0, 0))*m);
+	}
+	else
+	{
+		FMatrix m;
+		m.SetIdentity();
+		InCanvas->PushAbsoluteTransform(m);
+	}
+}
+
+void FOculusRiftHMD::PushViewCanvas(EStereoscopicPass StereoPass, FCanvas *InCanvas, UCanvas *InCanvasObject, FSceneView *InView) const
+{
+	if (StereoPass != eSSP_FULL)
+	{
+		if (InCanvasObject)
+		{
+			//InCanvasObject->Init();
+			InCanvasObject->SizeX = InView->ViewRect.Width();
+			InCanvasObject->SizeY = InView->ViewRect.Height();
+			InCanvasObject->SetView(InView);
+			InCanvasObject->Update();
+		}
+
+		InCanvas->PushAbsoluteTransform(FTranslationMatrix(FVector(InView->ViewRect.Min.X, InView->ViewRect.Min.Y, 0)));
+	}
+	else
+	{
+		FMatrix m;
+		m.SetIdentity();
+		InCanvas->PushAbsoluteTransform(m);
+	}
+}
+
+
+//---------------------------------------------------
+// Oculus Rift ISceneViewExtension Implementation
+//---------------------------------------------------
+
+void FOculusRiftHMD::ModifyShowFlags(FEngineShowFlags& ShowFlags)
+{
+	ShowFlags.MotionBlur = 0;
+#ifndef OVR_DIRECT_RENDERING
+	ShowFlags.HMDDistortion = bHmdDistortion;
+#else
+	ShowFlags.HMDDistortion = false;
+#endif
+	ShowFlags.ScreenPercentage = true;
+	ShowFlags.StereoRendering = IsStereoEnabled();
+}
+
+void FOculusRiftHMD::SetupView(FSceneViewFamily& InViewFamily, FSceneView& InView)
+{
+	InView.BaseHmdOrientation = LastHmdOrientation;
+	InView.BaseHmdLocation = FVector(0.f);
+	if (!bWorldToMetersOverride)
+	{
+		WorldToMetersScale = InView.WorldToMetersScale;
+	}
+
+#ifndef OVR_DIRECT_RENDERING
+	InViewFamily.bUseSeparateRenderTarget = false;
+#else
+	InViewFamily.bUseSeparateRenderTarget = ShouldUseSeparateRenderTarget();
+#endif
+
+	// check and save texture size. 
+	if (InView.StereoPass == eSSP_LEFT_EYE)
+	{
+		if (EyeViewportSize != InView.ViewRect.Size())
+		{
+			EyeViewportSize = InView.ViewRect.Size();
+			bNeedUpdateStereoRenderingParams = true;
+		}
+	}
+}
+
+bool FOculusRiftHMD::IsHeadTrackingAllowed() const
+{
+	return bHeadTrackingEnforced || GEngine->IsStereoscopic3D();
+}
+
+//---------------------------------------------------
+// Oculus Rift Specific
+//---------------------------------------------------
+
+FOculusRiftHMD::FOculusRiftHMD()
+	:
+	InitStatus(0)
+	, bStereoEnabled(false)
+	, bHMDEnabled(true)
+	, bNeedUpdateStereoRenderingParams(true)
+    , bOverrideStereo(false)
+	, bOverrideIPD(false)
+	, bOverrideDistortion(false)
+	, bDevSettingsEnabled(false)
+	, bOverrideFOV(false)
+	, bOverrideVSync(true)
+	, bVSync(true)
+	, bSavedVSync(false)
+	, SavedScrPerc(100.f)
+	, bOverrideScreenPercentage(false)
+	, ScreenPercentage(100.f)
+	, bAllowFinishCurrentFrame(true)
+	, InterpupillaryDistance(OVR_DEFAULT_IPD)
+	, WorldToMetersScale(100.f)
+	, bWorldToMetersOverride(false)
+	, UserDistanceToScreenModifier(0.f)
+	, HFOVInRadians(FMath::DegreesToRadians(90.f))
+	, VFOVInRadians(FMath::DegreesToRadians(90.f))
+	, MotionPredictionInSeconds(0)
+	, AccelGain(0.f)
+	, bHmdDistortion(true)
+	, bChromaAbCorrectionEnabled(true)
+	, bYawDriftCorrectionEnabled(true)
+	, bOverride2D(false)
+	, HudOffset(0.f)
+	, CanvasCenterOffset(0.f)
+	, bLowPersistenceMode(true) // on by default (DK2+ only)
+	, bUpdateOnRT(true)
+	, bOverdrive(true)
+	, bHQDistortion(true)
+	, bHeadTrackingEnforced(false)
+	, bMirrorToWindow(true)
+#if !UE_BUILD_SHIPPING
+	, bDrawTrackingCameraFrustum(false)
+	, bDoNotUpdateOnGT(false)
+	, bShowStats(false)
+	, bDrawGrid(false)
+	, bProfiling(false)
+#endif
+	, bTimeWarp(true)
+	, NearClippingPlane(0)
+	, FarClippingPlane(0)
+	, CurHmdOrientation(FQuat::Identity)
+	, DeltaControlRotation(FRotator::ZeroRotator)
+	, DeltaControlOrientation(FQuat::Identity)
+	, CurHmdPosition(FVector::ZeroVector)
+	, LastHmdOrientation(FQuat::Identity)
+	, LastHmdPosition(FVector::ZeroVector)
+	, BaseOffset(0, 0, 0)
+	, BaseOrientation(FQuat::Identity)
+	, Hmd(nullptr)
+	, TrackingCaps(0)
+	, DistortionCaps(0)
+	, HmdCaps(0)
+	, EyeViewportSize(0, 0)
+	, RenderParams_RenderThread(getThis())
+	, bHmdPosTracking(false)
+	, bHaveVisionTracking(false)
+{
+#ifdef OVR_VISION_ENABLED
+	bHmdPosTracking = true;
+#endif
+#ifndef OVR_DIRECT_RENDERING
+	bTimeWarp = false;
+#endif
+	SupportedTrackingCaps = SupportedDistortionCaps = SupportedHmdCaps = 0;
+	OSWindowHandle = nullptr;
+	Startup();
+}
+
+FOculusRiftHMD::~FOculusRiftHMD()
+{
+	Shutdown();
+}
+
+bool FOculusRiftHMD::IsInitialized() const
+{
+	return (InitStatus & eInitialized) != 0;
+}
+
+void FOculusRiftHMD::Startup()
+{
+	if (!IsRunningGame() || (InitStatus & eStartupExecuted) != 0)
+	{
+		// do not initialize plugin for server or if it was already initialized
+		return;
+	}
+	InitStatus |= eStartupExecuted;
+
+	// Initializes LibOVR. This LogMask_All enables maximum logging.
+	// Custom allocator can also be specified here.
+	// Actually, most likely, the ovr_Initialize is already called from PreInit.
+	ovr_Initialize();
+
+#if !UE_BUILD_SHIPPING
+	// Should be changed to CAPI when available.
+	static OculusLog OcLog;
+	OVR::Log::SetGlobalLog(&OcLog);
+#endif //#if !UE_BUILD_SHIPPING
+
+	Hmd = ovrHmd_Create(0);
+	if (Hmd)
+	{
+		InitStatus |= eInitialized;
+
+		// one frame prediction for the game thread. TODO: fugure out a better way to do this.
+		// In any case, the pose will be corrected by timewarp.
+		if (Hmd->Type >= ovrHmd_DK2)
+		{
+			MotionPredictionInSeconds = 1. / 75;
+		}
+		else
+		{
+			MotionPredictionInSeconds = 1. / 60;
+		}
+		SupportedDistortionCaps = Hmd->DistortionCaps;
+		SupportedHmdCaps		= Hmd->HmdCaps;
+		SupportedTrackingCaps	= Hmd->TrackingCaps;
+
+#ifndef OVR_DIRECT_RENDERING
+		SupportedDistortionCaps &= ~ovrDistortionCap_Overdrive;
+#endif
+#ifndef OVR_VISION_ENABLED
+		SupportedTrackingCaps &= ~ovrTrackingCap_Position;
+#endif
+
+		DistortionCaps	= SupportedDistortionCaps & (ovrDistortionCap_Chromatic | ovrDistortionCap_TimeWarp | ovrDistortionCap_Vignette | ovrDistortionCap_Overdrive);
+		TrackingCaps	= SupportedTrackingCaps & (ovrTrackingCap_Orientation | ovrTrackingCap_MagYawCorrection | ovrTrackingCap_Position);
+		HmdCaps			= SupportedHmdCaps & (ovrHmdCap_DynamicPrediction | ovrHmdCap_LowPersistence);
+		HmdCaps |= (bVSync ? 0 : ovrHmdCap_NoVSync);
+
+		if (!(SupportedDistortionCaps & ovrDistortionCap_TimeWarp))
+		{
+			bTimeWarp = false;
+		}
+
+		if (!bTimeWarp)
+		{
+			// Without Timewarp we may safely multiply the prediction by two.
+			MotionPredictionInSeconds *= 2;
+		}
+
+		UpdateDistortionCaps();
+		UpdateHmdRenderInfo();
+		UpdateStereoRenderingParams();
+		UE_LOG(LogHMD, Log, TEXT("Oculus initialized."));
+
+		// Uncap fps to enable FPS higher than 62
+		GEngine->bSmoothFrameRate = false;
+	}
+	else
+	{
+		UE_LOG(LogHMD, Warning, TEXT("No Oculus HMD detected! Is Oculus Run-Time installed and service is running?"));
+	}
+
+	LoadFromIni();
+	SaveSystemValues();
+
+	UpdateHmdCaps();
+
+#ifdef OVR_DIRECT_RENDERING
+#if defined(OVR_D3D_VERSION) && (OVR_D3D_VERSION == 11)
+	if (IsPCPlatform(GRHIShaderPlatform) && !IsOpenGLPlatform(GRHIShaderPlatform))
+	{
+		pD3D11Bridge = new D3D11Bridge(this);
+	}
+#endif
+#if defined(OVR_GL)
+	if (IsOpenGLPlatform(GRHIShaderPlatform))
+	{
+		pOGLBridge = new OGLBridge(this);
+	}
+#endif
+#endif // #ifdef OVR_DIRECT_RENDERING
+}
+
+void FOculusRiftHMD::Shutdown()
+{
+	ovrHmd_AttachToWindow(Hmd, NULL, NULL, NULL);
+
+	if (!(InitStatus & eStartupExecuted))
+	{
+		return;
+	}
+	SaveToIni();
+
+#ifdef OVR_DIRECT_RENDERING
+	ENQUEUE_UNIQUE_RENDER_COMMAND_ONEPARAMETER(ShutdownRen,
+	FOculusRiftHMD*, Plugin, this,
+	{
+		Plugin->ShutdownRendering();
+	});
+	// Wait for all resources to be released
+	FlushRenderingCommands();
+#else
+	ovrHmd_Destroy(Hmd);
+	Hmd = nullptr;
+#endif 
+	
+#ifndef OVR_DIRECT_RENDERING
+	for (unsigned i = 0; i < sizeof(pDistortionMesh) / sizeof(pDistortionMesh[0]); ++i)
+	{
+		pDistortionMesh[i] = NULL;
+	}
+#endif
+	{
+		OVR::Lock::Locker lock(&StereoParamsLock);
+		RenderParams_RenderThread.Clear();
+	}
+	ovr_Shutdown();
+	InitStatus = 0;
+	UE_LOG(LogHMD, Log, TEXT("Oculus shutdown."));
+}
+
+void FOculusRiftHMD::UpdateDistortionCaps()
+{
+	if (IsOpenGLPlatform(GRHIShaderPlatform))
+	{
+		DistortionCaps &= ~ovrDistortionCap_SRGB;
+		DistortionCaps |= ovrDistortionCap_FlipInput;
+	}
+	(bTimeWarp) ? DistortionCaps |= ovrDistortionCap_TimeWarp : DistortionCaps &= ~ovrDistortionCap_TimeWarp;
+	(bOverdrive) ? DistortionCaps |= ovrDistortionCap_Overdrive : DistortionCaps &= ~ovrDistortionCap_Overdrive;
+	//(bHQDistortion) ? DistortionCaps |= ovrDistortionCap_HqDistortion : DistortionCaps &= ~ovrDistortionCap_HqDistortion;
+	(bChromaAbCorrectionEnabled) ? DistortionCaps |= ovrDistortionCap_Chromatic : DistortionCaps &= ~ovrDistortionCap_Chromatic;
+#if !UE_BUILD_SHIPPING
+	(bProfiling) ? DistortionCaps |= ovrDistortionCap_ProfileNoTimewarpSpinWaits : DistortionCaps &= ~ovrDistortionCap_ProfileNoTimewarpSpinWaits;
+#endif // #if !UE_BUILD_SHIPPING
+
+#ifdef OVR_DIRECT_RENDERING 
+	if (GetActiveRHIBridgeImpl())
+	{
+		GetActiveRHIBridgeImpl()->SetNeedReinitRendererAPI();
+	}
+#endif // OVR_DIRECT_RENDERING
+}
+
+void FOculusRiftHMD::UpdateHmdCaps()
+{
+	if (Hmd)
+	{
+		TrackingCaps = ovrTrackingCap_Orientation;
+		if (bYawDriftCorrectionEnabled)
+		{
+			TrackingCaps |= ovrTrackingCap_MagYawCorrection;
+		}
+		else
+		{
+			TrackingCaps &= ~ovrTrackingCap_MagYawCorrection;
+		}
+		if (bHmdPosTracking)
+		{
+			TrackingCaps |= ovrTrackingCap_Position;
+		}
+		else
+		{
+			TrackingCaps &= ~ovrTrackingCap_Position;
+		}
+
+		if (bLowPersistenceMode)
+		{
+			HmdCaps |= ovrHmdCap_LowPersistence;
+		}
+		else
+		{
+			HmdCaps &= ~ovrHmdCap_LowPersistence;
+		}
+
+		if (bVSync)
+		{
+			HmdCaps &= ~ovrHmdCap_NoVSync;
+		}
+		else
+		{
+			HmdCaps |= ovrHmdCap_NoVSync;
+		}
+
+		if (bMirrorToWindow)
+		{
+			HmdCaps &= ~ovrHmdCap_NoMirrorToWindow;
+		}
+		else
+		{
+			HmdCaps |= ovrHmdCap_NoMirrorToWindow;
+		}
+		ovrHmd_SetEnabledCaps(Hmd, HmdCaps);
+
+		ovrHmd_ConfigureTracking(Hmd, TrackingCaps, 0);
+	}
+}
+
+FORCEINLINE static float GetVerticalFovRadians(const ovrFovPort& fov)
+{
+	return FMath::Atan(fov.UpTan) + FMath::Atan(fov.DownTan);
+}
+
+FORCEINLINE static float GetHorizontalFovRadians(const ovrFovPort& fov)
+{
+	return FMath::Atan(fov.LeftTan) + FMath::Atan(fov.RightTan);
+}
+
+void FOculusRiftHMD::UpdateHmdRenderInfo()
+{
+	check(Hmd);
+
+	UE_LOG(LogHMD, Warning, TEXT("HMD %s, Monitor %s, res = %d x %d, windowPos = {%d, %d}"), ANSI_TO_TCHAR(Hmd->ProductName), 
+		ANSI_TO_TCHAR(Hmd->DisplayDeviceName), Hmd->Resolution.w, Hmd->Resolution.h, Hmd->WindowsPos.x, Hmd->WindowsPos.y); 
+
+	// Calc FOV
+	if (!bOverrideFOV)
+	{
+		// Calc FOV, symmetrical, for each eye. 
+		EyeFov[0] = Hmd->DefaultEyeFov[0];
+		EyeFov[1] = Hmd->DefaultEyeFov[1];
+
+		// Calc FOV in radians
+		VFOVInRadians = FMath::Max(GetVerticalFovRadians(EyeFov[0]), GetVerticalFovRadians(EyeFov[1]));
+		HFOVInRadians = FMath::Max(GetHorizontalFovRadians(EyeFov[0]), GetHorizontalFovRadians(EyeFov[1]));
+	}
+
+	const Sizei recommenedTex0Size = ovrHmd_GetFovTextureSize(Hmd, ovrEye_Left, EyeFov[0], 1.0f);
+	const Sizei recommenedTex1Size = ovrHmd_GetFovTextureSize(Hmd, ovrEye_Right, EyeFov[1], 1.0f);
+
+	Sizei idealRenderTargetSize;
+	idealRenderTargetSize.w = recommenedTex0Size.w + recommenedTex1Size.w;
+	idealRenderTargetSize.h = FMath::Max(recommenedTex0Size.h, recommenedTex1Size.h);
+
+	IdealScreenPercentage = FMath::Max(float(idealRenderTargetSize.w) / float(Hmd->Resolution.w) * 100.f,
+									   float(idealRenderTargetSize.h) / float(Hmd->Resolution.h) * 100.f);
+
+	// Override eye distance by the value from HMDInfo (stored in Profile).
+	if (!bOverrideIPD)
+	{
+		InterpupillaryDistance = ovrHmd_GetFloat(Hmd, OVR_KEY_IPD, OVR_DEFAULT_IPD);
+	}
+
+	// Default texture size (per eye) is equal to half of W x H resolution. Will be overridden in SetupView.
+	EyeViewportSize = FIntPoint(Hmd->Resolution.w / 2, Hmd->Resolution.h);
+
+	bNeedUpdateStereoRenderingParams = true;
+}
+
+void FOculusRiftHMD::UpdateStereoRenderingParams()
+{
+	// If we've manually overridden stereo rendering params for debugging, don't mess with them
+	if (bOverrideStereo || !IsStereoEnabled())
+	{
+		return;
+	}
+	if (IsInitialized())
+	{
+		Lock::Locker lock(&StereoParamsLock);
+
+		TextureSize = Sizei(EyeViewportSize.X * 2, EyeViewportSize.Y);
+
+		EyeRenderViewport[0].Pos = Vector2i(0, 0);
+		EyeRenderViewport[0].Size = Sizei(EyeViewportSize.X, EyeViewportSize.Y);
+		EyeRenderViewport[1].Pos = Vector2i(EyeViewportSize.X, 0);
+		EyeRenderViewport[1].Size = EyeRenderViewport[0].Size;
+
+		//!AB: note, for Direct Rendering EyeRenderDesc is calculated twice, once
+		// here and another time in BeginRendering_RenderThread. I need to have EyeRenderDesc
+		// on a game thread for ViewAdjust (for ProjectionMatrix calculation). @@TODO: revise.
+		EyeRenderDesc[0] = ovrHmd_GetRenderDesc(Hmd, ovrEye_Left, EyeFov[0]);
+		EyeRenderDesc[1] = ovrHmd_GetRenderDesc(Hmd, ovrEye_Right, EyeFov[1]);
+
+		const bool bRightHanded = false;
+		// Far and Near clipping planes will be modified in GetStereoProjectionMatrix()
+		EyeProjectionMatrices[0] = ovrMatrix4f_Projection(EyeFov[0], 0.01f, 10000.0f, bRightHanded);
+		EyeProjectionMatrices[1] = ovrMatrix4f_Projection(EyeFov[1], 0.01f, 10000.0f, bRightHanded);
+
+		// 2D elements offset
+		if (!bOverride2D)
+		{
+			float ScreenSizeInMeters[2]; // 0 - width, 1 - height
+			float LensSeparationInMeters;
+			LensSeparationInMeters = ovrHmd_GetFloat(Hmd, "LensSeparation", 0);
+			ovrHmd_GetFloatArray(Hmd, "ScreenSize", ScreenSizeInMeters, 2);
+
+			// Recenter projection (meters)
+			const float LeftProjCenterM = ScreenSizeInMeters[0] * 0.25f;
+			const float LensRecenterM = LeftProjCenterM - LensSeparationInMeters * 0.5f;
+
+			// Recenter projection (normalized)
+			const float LensRecenter = 4.0f * LensRecenterM / ScreenSizeInMeters[0];
+
+			HudOffset = 0.25f * InterpupillaryDistance * (Hmd->Resolution.w / ScreenSizeInMeters[0]) / 15.0f;
+			CanvasCenterOffset = (0.25f * LensRecenter) * Hmd->Resolution.w;
+		}
+
+		PrecalculatePostProcess_NoLock();
+#ifdef OVR_DIRECT_RENDERING 
+		GetActiveRHIBridgeImpl()->SetNeedReinitRendererAPI();
+#endif // OVR_DIRECT_RENDERING
+	}
+	else
+	{
+		CanvasCenterOffset = 0.f;
+	}
+	bNeedUpdateStereoRenderingParams = false;
+}
+
+void FOculusRiftHMD::LoadFromIni()
+{
+	const TCHAR* OculusSettings = TEXT("Oculus.Settings");
+	bool v;
+	float f;
+	if (GConfig->GetBool(OculusSettings, TEXT("bChromaAbCorrectionEnabled"), v, GEngineIni))
+	{
+		bChromaAbCorrectionEnabled = v;
+	}
+	if (GConfig->GetBool(OculusSettings, TEXT("bYawDriftCorrectionEnabled"), v, GEngineIni))
+	{
+		bYawDriftCorrectionEnabled = v;
+	}
+	if (GConfig->GetBool(OculusSettings, TEXT("bDevSettingsEnabled"), v, GEngineIni))
+	{
+		bDevSettingsEnabled = v;
+	}
+	if (GConfig->GetBool(OculusSettings, TEXT("bOverrideIPD"), v, GEngineIni))
+	{
+		bOverrideIPD = v;
+		if (bOverrideIPD)
+		{
+			if (GConfig->GetFloat(OculusSettings, TEXT("IPD"), f, GEngineIni))
+			{
+				SetInterpupillaryDistance(f);
+			}
+		}
+	}
+	if (GConfig->GetBool(OculusSettings, TEXT("bOverrideStereo"), v, GEngineIni))
+	{
+		bOverrideStereo = v;
+		if (bOverrideStereo)
+		{
+			if (GConfig->GetFloat(OculusSettings, TEXT("HFOV"), f, GEngineIni))
+			{
+				HFOVInRadians = f;
+			}
+			if (GConfig->GetFloat(OculusSettings, TEXT("VFOV"), f, GEngineIni))
+			{
+				VFOVInRadians = f;
+			}
+		}
+	}
+	if (GConfig->GetBool(OculusSettings, TEXT("bOverrideVSync"), v, GEngineIni))
+	{
+		bOverrideVSync = v;
+		if (GConfig->GetBool(OculusSettings, TEXT("bVSync"), v, GEngineIni))
+		{
+			bVSync = v;
+		}
+	}
+	if (GConfig->GetBool(OculusSettings, TEXT("bOverrideScreenPercentage"), v, GEngineIni))
+	{
+		bOverrideScreenPercentage = v;
+		if (GConfig->GetFloat(OculusSettings, TEXT("ScreenPercentage"), f, GEngineIni))
+		{
+			ScreenPercentage = f;
+		}
+	}
+	if (GConfig->GetBool(OculusSettings, TEXT("bAllowFinishCurrentFrame"), v, GEngineIni))
+	{
+		bAllowFinishCurrentFrame = v;
+	}
+#ifdef OVR_VISION_ENABLED
+	if (GConfig->GetBool(OculusSettings, TEXT("bHmdPosTracking"), v, GEngineIni))
+	{
+		bHmdPosTracking = v;
+	}
+#endif // #ifdef OVR_VISION_ENABLED
+	if (GConfig->GetBool(OculusSettings, TEXT("bLowPersistenceMode"), v, GEngineIni))
+	{
+		bLowPersistenceMode = v;
+	}
+	if (GConfig->GetBool(OculusSettings, TEXT("bUpdateOnRT"), v, GEngineIni))
+	{
+		bUpdateOnRT = v;
+	}
+	if (GConfig->GetFloat(OculusSettings, TEXT("FarClippingPlane"), f, GEngineIni))
+	{
+		FarClippingPlane = f;
+	}
+	if (GConfig->GetFloat(OculusSettings, TEXT("NearClippingPlane"), f, GEngineIni))
+	{
+		NearClippingPlane = f;
+	}
+}
+
+void FOculusRiftHMD::SaveToIni()
+{
+	const TCHAR* OculusSettings = TEXT("Oculus.Settings");
+	GConfig->SetBool(OculusSettings, TEXT("bChromaAbCorrectionEnabled"), bChromaAbCorrectionEnabled, GEngineIni);
+	GConfig->SetBool(OculusSettings, TEXT("bYawDriftCorrectionEnabled"), bYawDriftCorrectionEnabled, GEngineIni);
+	GConfig->SetBool(OculusSettings, TEXT("bDevSettingsEnabled"), bDevSettingsEnabled, GEngineIni);
+
+	GConfig->SetBool(OculusSettings, TEXT("bOverrideIPD"), bOverrideIPD, GEngineIni);
+	if (bOverrideIPD)
+	{
+		GConfig->SetFloat(OculusSettings, TEXT("IPD"), GetInterpupillaryDistance(), GEngineIni);
+	}
+	GConfig->SetBool(OculusSettings, TEXT("bOverrideStereo"), bOverrideStereo, GEngineIni);
+	if (bOverrideStereo)
+	{
+		GConfig->SetFloat(OculusSettings, TEXT("HFOV"), HFOVInRadians, GEngineIni);
+		GConfig->SetFloat(OculusSettings, TEXT("VFOV"), VFOVInRadians, GEngineIni);
+	}
+
+	GConfig->SetBool(OculusSettings, TEXT("bOverrideVSync"), bOverrideVSync, GEngineIni);
+	if (bOverrideVSync)
+	{
+		GConfig->SetBool(OculusSettings, TEXT("VSync"), bVSync, GEngineIni);
+	}
+
+	GConfig->SetBool(OculusSettings, TEXT("bOverrideScreenPercentage"), bOverrideScreenPercentage, GEngineIni);
+	if (bOverrideScreenPercentage)
+	{
+		// Save the current ScreenPercentage state
+		GConfig->SetFloat(OculusSettings, TEXT("ScreenPercentage"), ScreenPercentage, GEngineIni);
+	}
+	GConfig->SetBool(OculusSettings, TEXT("bAllowFinishCurrentFrame"), bAllowFinishCurrentFrame, GEngineIni);
+
+#ifdef OVR_VISION_ENABLED
+	GConfig->SetBool(OculusSettings, TEXT("bHmdPosTracking"), bHmdPosTracking, GEngineIni);
+#endif
+	GConfig->SetBool(OculusSettings, TEXT("bLowPersistenceMode"), bLowPersistenceMode, GEngineIni);
+
+	GConfig->SetBool(OculusSettings, TEXT("bUpdateOnRT"), bUpdateOnRT, GEngineIni);
+
+	GConfig->SetFloat(OculusSettings, TEXT("FarClippingPlane"), FarClippingPlane, GEngineIni);
+	GConfig->SetFloat(OculusSettings, TEXT("NearClippingPlane"), NearClippingPlane, GEngineIni);
+}
+
+bool FOculusRiftHMD::HandleInputKey(UPlayerInput* pPlayerInput,
+	const FKey& Key, EInputEvent EventType, float AmountDepressed, bool bGamepad)
+{
+	if (EventType == IE_Pressed && Hmd && IsStereoEnabled())
+	{
+		if (!Key.IsMouseButton())
+		{
+			ovrHmd_DismissHSWDisplay(Hmd);
+		}
+	}
+	return false;
+}
+
+#endif //OCULUS_RIFT_SUPPORTED_PLATFORMS
+
diff --git a/Engine/Plugins/Runtime/OculusRift/Source/OculusRift/Private/OculusRiftHMD.h b/Engine/Plugins/Runtime/OculusRift/Source/OculusRift/Private/OculusRiftHMD.h
new file mode 100644
--- /dev/null
+++ b/Engine/Plugins/Runtime/OculusRift/Source/OculusRift/Private/OculusRiftHMD.h
@@ -0,0 +1,638 @@
+// Copyright 1998-2014 Epic Games, Inc. All Rights Reserved.
+
+#pragma once
+#include "IOculusRiftPlugin.h"
+#include "HeadMountedDisplay.h"
+#include "IHeadMountedDisplay.h"
+
+#if OCULUS_RIFT_SUPPORTED_PLATFORMS
+
+#include "SceneViewExtension.h"
+
+#if PLATFORM_SUPPORTS_PRAGMA_PACK
+#pragma pack (push,8)
+#endif
+
+#if PLATFORM_WINDOWS
+	#define OVR_VISION_ENABLED
+	#define OVR_DIRECT_RENDERING
+	#define OVR_D3D_VERSION 11
+	#define OVR_GL
+#elif PLATFORM_MAC
+	#define OVR_VISION_ENABLED
+    #define OVR_DIRECT_RENDERING
+    #define OVR_GL
+#endif
+
+#ifdef OVR_VISION_ENABLED
+	#ifndef OVR_CAPI_VISIONSUPPORT
+		#define OVR_CAPI_VISIONSUPPORT
+	#endif
+#endif
+
+#if OCULUS_RIFT_SUPPORTED_PLATFORMS
+	
+	#include "OVR.h"
+	#include "OVR_Version.h"
+	#include "OVR_Kernel.h"
+
+	#include "../Src/Kernel/OVR_Threads.h"
+
+#ifdef OVR_DIRECT_RENDERING
+    #if PLATFORM_WINDOWS
+        #include "AllowWindowsPlatformTypes.h"
+    #endif
+	#ifdef OVR_D3D_VERSION
+		#include "../Src/OVR_CAPI_D3D.h"
+	#endif // OVR_D3D_VERSION
+	#ifdef OVR_GL
+		#include "../Src/OVR_CAPI_GL.h"
+	#endif
+#endif // OVR_DIRECT_RENDERING
+
+	using namespace OVR;
+
+#endif //OCULUS_RIFT_SUPPORTED_PLATFORMS
+
+#if PLATFORM_SUPPORTS_PRAGMA_PACK
+#pragma pack (pop)
+#endif
+
+
+/**
+ * Oculus Rift Head Mounted Display
+ */
+class FOculusRiftHMD : public IHeadMountedDisplay, public ISceneViewExtension
+{
+public:
+	static void PreInit();
+
+	/** IHeadMountedDisplay interface */
+	virtual bool IsHMDEnabled() const override;
+	virtual void EnableHMD(bool allow = true) override;
+	virtual EHMDDeviceType::Type GetHMDDeviceType() const override;
+	virtual bool GetHMDMonitorInfo(MonitorInfo&) const override;
+
+	virtual bool DoesSupportPositionalTracking() const override;
+	virtual bool HasValidTrackingPosition() const override;
+	virtual void GetPositionalTrackingCameraProperties(FVector& OutOrigin, FRotator& OutOrientation, float& OutHFOV, float& OutVFOV, float& OutCameraDistance, float& OutNearPlane, float& OutFarPlane) const override;
+
+	virtual void SetInterpupillaryDistance(float NewInterpupillaryDistance) override;
+	virtual float GetInterpupillaryDistance() const override;
+    //virtual float GetFieldOfViewInRadians() const override;
+	virtual void GetFieldOfView(float& OutHFOVInDegrees, float& OutVFOVInDegrees) const override;
+
+	virtual void GetCurrentOrientationAndPosition(FQuat& CurrentOrientation, FVector& CurrentPosition) const override;
+	virtual void ApplyHmdRotation(APlayerController* PC, FRotator& ViewRotation) override;
+	virtual void UpdatePlayerCameraRotation(APlayerCameraManager*, struct FMinimalViewInfo& POV) override;
+
+	virtual bool IsChromaAbCorrectionEnabled() const override;
+
+	virtual class ISceneViewExtension* GetViewExtension() override;
+	virtual bool Exec( UWorld* InWorld, const TCHAR* Cmd, FOutputDevice& Ar ) override;
+	virtual void OnScreenModeChange(EWindowMode::Type WindowMode) override;
+
+	virtual bool IsFullScreenAllowed() const override;
+	virtual void RecordAnalytics() override;
+
+	/** IStereoRendering interface */
+	virtual bool IsStereoEnabled() const override;
+	virtual bool EnableStereo(bool stereo = true) override;
+    virtual void AdjustViewRect(EStereoscopicPass StereoPass, int32& X, int32& Y, uint32& SizeX, uint32& SizeY) const override;
+	virtual void CalculateStereoViewOffset(const EStereoscopicPass StereoPassType, const FRotator& ViewRotation, 
+										   const float MetersToWorld, FVector& ViewLocation) override;
+	virtual FMatrix GetStereoProjectionMatrix(const EStereoscopicPass StereoPassType, const float FOV) const override;
+	virtual void InitCanvasFromView(FSceneView* InView, UCanvas* Canvas) override;
+	virtual void PushViewportCanvas(EStereoscopicPass StereoPass, FCanvas *InCanvas, UCanvas *InCanvasObject, FViewport *InViewport) const override;
+	virtual void PushViewCanvas(EStereoscopicPass StereoPass, FCanvas *InCanvas, UCanvas *InCanvasObject, FSceneView *InView) const override;
+	virtual void GetEyeRenderParams_RenderThread(EStereoscopicPass StereoPass, FVector2D& EyeToSrcUVScaleValue, FVector2D& EyeToSrcUVOffsetValue) const override;
+	virtual void GetTimewarpMatrices_RenderThread(EStereoscopicPass StereoPass, FMatrix& EyeRotationStart, FMatrix& EyeRotationEnd) const override;
+
+	virtual void UpdateViewport(bool bUseSeparateRenderTarget, const FViewport& Viewport) override;
+
+#ifdef OVR_DIRECT_RENDERING
+	virtual void CalculateRenderTargetSize(uint32& InOutSizeX, uint32& InOutSizeY) const override;
+	virtual bool NeedReAllocateViewportRenderTarget(const FViewport& Viewport) const override;
+#endif//OVR_DIRECT_RENDERING
+
+	virtual bool ShouldUseSeparateRenderTarget() const override
+	{
+#ifdef OVR_DIRECT_RENDERING
+		check(IsInGameThread());
+		return IsStereoEnabled();
+#else
+		return false;
+#endif//OVR_DIRECT_RENDERING
+	}
+
+    /** ISceneViewExtension interface */
+    virtual void ModifyShowFlags(FEngineShowFlags& ShowFlags) override;
+    virtual void SetupView(FSceneViewFamily& InViewFamily, FSceneView& InView) override;
+    virtual void PreRenderView_RenderThread(FSceneView& InView) override;
+	virtual void PreRenderViewFamily_RenderThread(FSceneViewFamily& InViewFamily) override;
+
+	/** Positional tracking control methods */
+	virtual bool IsPositionalTrackingEnabled() const override;
+	virtual bool EnablePositionalTracking(bool enable) override;
+
+	virtual bool IsHeadTrackingAllowed() const override;
+
+	virtual bool IsInLowPersistenceMode() const override;
+	virtual void EnableLowPersistenceMode(bool Enable = true) override;
+
+	/** Resets orientation by setting roll and pitch to 0, 
+	    assuming that current yaw is forward direction and assuming
+		current position as 0 point. */
+	virtual void ResetOrientationAndPosition(float yaw = 0.f) override;
+
+	virtual void DrawDistortionMesh_RenderThread(struct FRenderingCompositePassContext& Context, const FSceneView& View, const FIntPoint& TextureSize) override;
+	virtual void UpdateScreenSettings(const FViewport*) override;
+
+	virtual bool HandleInputKey(class UPlayerInput*, const FKey& Key, EInputEvent EventType, float AmountDepressed, bool bGamepad) override;
+
+	virtual void DrawDebug(UCanvas* Canvas, EStereoscopicPass StereoPass) override;
+
+	void GetCurrentPose(FQuat& CurrentOrientation, FVector& CurrentPosition) const;
+	void BeginRendering_RenderThread();
+
+#ifdef OVR_DIRECT_RENDERING
+	class BridgeBaseImpl : public FRHICustomPresent
+	{
+	public:
+		BridgeBaseImpl(FOculusRiftHMD* plugin) :
+			FRHICustomPresent(nullptr),
+			Plugin(plugin), 
+			bNeedReinitRendererAPI(true), 
+			bInitialized(false) 
+		{}
+
+		// Returns true if it is initialized and used.
+		bool IsInitialized() const { return bInitialized; }
+
+		virtual void BeginRendering() = 0;
+		virtual void UpdateViewport(const FViewport& Viewport, FRHIViewport* ViewportRHI) = 0;
+		virtual void SetNeedReinitRendererAPI() { bNeedReinitRendererAPI = true; }
+
+		virtual void Reset() = 0;
+		virtual void Shutdown() = 0;
+
+	protected: // data
+		// Data
+		mutable OVR::Lock	ModifyLock;
+		mutable OVR::Lock	ModifyEyeTexturesLock;
+		FOculusRiftHMD*		Plugin;
+		bool				bNeedReinitRendererAPI;
+		bool				bInitialized;
+	};
+
+#if defined(OVR_D3D_VERSION) && (OVR_D3D_VERSION == 11)
+	class D3D11Bridge : public BridgeBaseImpl
+	{
+	public:
+		D3D11Bridge(FOculusRiftHMD* plugin);
+
+		// Implementation of FRHICustomPresent
+		// Resets Viewport-specific pointers (BackBufferRT, SwapChain).
+		virtual void OnBackBufferResize() override;
+		// Returns true if Engine should perform its own Present.
+		virtual bool Present(int SyncInterval) override;
+
+		// Implementation of BridgeBaseImpl, called by Plugin itself
+		virtual void BeginRendering() override;
+		void FinishRendering();
+		virtual void UpdateViewport(const FViewport& Viewport, FRHIViewport* ViewportRHI) override;
+		virtual void Reset() override;
+		virtual void Shutdown() override
+		{
+			Reset();
+		}
+
+	protected:
+		void Reset_RenderThread();
+	protected: // data
+		ovrD3D11Config		Cfg;
+		ovrD3D11Texture		EyeTexture[2];				
+		ovrD3D11Texture		EyeTexture_RenderThread[2];
+		bool				bNeedReinitEyeTextures;
+	};
+
+#endif
+
+#ifdef OVR_GL
+	class OGLBridge : public BridgeBaseImpl
+	{
+	public:
+		OGLBridge(FOculusRiftHMD* plugin);
+
+		// Implementation of FRHICustomPresent
+		// Resets Viewport-specific resources.
+		virtual void OnBackBufferResize() override;
+		// Returns true if Engine should perform its own Present.
+		virtual bool Present(int SyncInterval) override;
+
+		// Implementation of BridgeBaseImpl, called by Plugin itself
+		virtual void BeginRendering() override;
+		void FinishRendering();
+		virtual void UpdateViewport(const FViewport& Viewport, FRHIViewport* ViewportRHI) override;
+		virtual void Reset() override;
+		virtual void Shutdown() override
+		{
+			Reset();
+		}
+
+		virtual void Init();
+
+	protected: // data
+		ovrGLConfig			Cfg;
+		ovrGLTexture		EyeTexture[2];
+		ovrGLTexture		EyeTexture_RenderThread[2];
+		bool				bNeedReinitEyeTextures;
+	};
+#endif // OVR_GL
+	BridgeBaseImpl* GetActiveRHIBridgeImpl();
+
+	void ShutdownRendering();
+
+#else
+
+	virtual void FinishRenderingFrame_RenderThread(FRHICommandListImmediate& RHICmdList) override;
+
+#endif // #ifdef OVR_DIRECT_RENDERING
+
+	/** Constructor */
+	FOculusRiftHMD();
+
+	/** Destructor */
+	virtual ~FOculusRiftHMD();
+
+	/** @return	True if the HMD was initialized OK */
+	bool IsInitialized() const;
+
+private:
+	FOculusRiftHMD* getThis() { return this; }
+
+	/**
+	 * Starts up the Oculus Rift device
+	 */
+	void Startup();
+
+	/**
+	 * Shuts down Oculus Rift
+	 */
+	void Shutdown();
+
+	/**
+	 * Reads the device configuration, and sets up the stereoscopic rendering parameters
+	 */
+	void UpdateStereoRenderingParams();
+	void UpdateHmdRenderInfo();
+
+    /**
+     * Updates the view point reflecting the current HMD orientation. 
+     */
+	static void UpdatePlayerViewPoint(const FQuat& CurrentOrientation, const FVector& CurrentPosition, const FVector& LastHmdPosition, const FQuat& DeltaControlOrientation, const FQuat& BaseViewOrientation, const FVector& BaseViewPosition, FRotator& ViewRotation, FVector& ViewLocation);
+
+	/**
+	 * Converts quat from Oculus ref frame to Unreal
+	 */
+	template <typename OVRQuat>
+	FORCEINLINE FQuat ToFQuat(const OVRQuat& InQuat) const
+	{
+		return FQuat(float(-InQuat.z), float(InQuat.x), float(InQuat.y), float(-InQuat.w));
+	}
+	/**
+	 * Converts vector from Oculus to Unreal
+	 */
+	template <typename OVRVector3>
+	FORCEINLINE FVector ToFVector(const OVRVector3& InVec) const
+	{
+		return FVector(float(-InVec.z), float(InVec.x), float(InVec.y));
+	}
+
+	/**
+	 * Converts vector from Oculus to Unreal, also converting meters to UU (Unreal Units)
+	 */
+	template <typename OVRVector3>
+	FORCEINLINE FVector ToFVector_M2U(const OVRVector3& InVec) const
+	{
+		return FVector(float(-InVec.z * WorldToMetersScale), 
+			           float(InVec.x  * WorldToMetersScale), 
+					   float(InVec.y  * WorldToMetersScale));
+	}
+	/**
+	 * Converts vector from Oculus to Unreal, also converting UU (Unreal Units) to meters.
+	 */
+	template <typename OVRVector3>
+	FORCEINLINE OVRVector3 ToOVRVector_U2M(const FVector& InVec) const
+	{
+		return OVRVector3(float(InVec.Y * (1.f / WorldToMetersScale)), 
+			              float(InVec.Z * (1.f / WorldToMetersScale)), 
+					     float(-InVec.X * (1.f / WorldToMetersScale)));
+	}
+	/**
+	 * Converts vector from Oculus to Unreal.
+	 */
+	template <typename OVRVector3>
+	FORCEINLINE OVRVector3 ToOVRVector(const FVector& InVec) const
+	{
+		return OVRVector3(float(InVec.Y), float(InVec.Z), float(-InVec.X));
+	}
+
+	/** Converts vector from meters to UU (Unreal Units) */
+	FORCEINLINE FVector MetersToUU(const FVector& InVec) const
+	{
+		return InVec * WorldToMetersScale;
+	}
+	/** Converts vector from UU (Unreal Units) to meters */
+	FORCEINLINE FVector UUToMeters(const FVector& InVec) const
+	{
+		return InVec * (1.f / WorldToMetersScale);
+	}
+
+	FORCEINLINE FMatrix ToFMatrix(const OVR::Matrix4f& vtm) const
+	{
+		// Rows and columns are swapped between OVR::Matrix4f and FMatrix
+		return FMatrix(
+			FPlane(vtm.M[0][0], vtm.M[1][0], vtm.M[2][0], vtm.M[3][0]),
+			FPlane(vtm.M[0][1], vtm.M[1][1], vtm.M[2][1], vtm.M[3][1]),
+			FPlane(vtm.M[0][2], vtm.M[1][2], vtm.M[2][2], vtm.M[3][2]),
+			FPlane(vtm.M[0][3], vtm.M[1][3], vtm.M[2][3], vtm.M[3][3]));
+	}
+
+	/**
+	 * Called when state changes from 'stereo' to 'non-stereo'. Suppose to distribute
+	 * the event further to user's code (?).
+	 */
+	void OnOculusStateChange(bool bIsEnabledNow);
+
+	/** Load/save settings */
+	void LoadFromIni();
+	void SaveToIni();
+
+	/** Applies overrides on system when switched to stereo (such as VSync and ScreenPercentage) */
+	void ApplySystemOverridesOnStereo(bool force = false);
+	/** Saves system values before applying overrides. */
+	void SaveSystemValues();
+	/** Restores system values after overrides applied. */
+	void RestoreSystemValues();
+
+	void ResetControlRotation() const;
+
+	void PrecalculatePostProcess_NoLock();
+
+	void UpdateDistortionCaps();
+	void UpdateHmdCaps();
+
+	void PoseToOrientationAndPosition(const ovrPosef& InPose, FQuat& OutOrientation, FVector& OutPosition) const;
+
+#if !UE_BUILD_SHIPPING
+	void DrawDebugTrackingCameraFrustum(class UWorld* InWorld, const FVector& ViewLocation);
+#endif // #if !UE_BUILD_SHIPPING
+
+private: // data
+	friend class FOculusMessageHandler;
+
+	/** Whether or not the Oculus was successfully initialized */
+	enum EInitStatus
+	{
+		eNotInitialized   = 0x00,
+		eStartupExecuted  = 0x01,
+		eInitialized      = 0x02,
+	};
+	int InitStatus; // see bitmask EInitStatus
+
+	/** Whether stereo is currently on or off. */
+	bool bStereoEnabled;
+
+	/** Whether or not switching to stereo is allowed */
+	bool bHMDEnabled;
+
+	/** Indicates if it is necessary to update stereo rendering params */
+	bool bNeedUpdateStereoRenderingParams;
+
+	/** Debugging:  Whether or not the stereo rendering settings have been manually overridden by an exec command.  They will no longer be auto-calculated */
+    bool bOverrideStereo;
+
+	/** Debugging:  Whether or not the IPD setting have been manually overridden by an exec command. */
+	bool bOverrideIPD;
+
+	/** Debugging:  Whether or not the distortion settings have been manually overridden by an exec command.  They will no longer be auto-calculated */
+	bool bOverrideDistortion;
+
+	/** Debugging: Allows changing internal params, such as screen size, eye-to-screen distance, etc */
+	bool bDevSettingsEnabled;
+
+	bool bOverrideFOV;
+
+	/** Whether or not to override game VSync setting when switching to stereo */
+	bool bOverrideVSync;
+	/** Overridden VSync value */
+	bool bVSync;
+
+	/** Saved original values for VSync and ScreenPercentage. */
+	bool bSavedVSync;
+	float SavedScrPerc;
+
+	/** Whether or not to override game ScreenPercentage setting when switching to stereo */
+	bool bOverrideScreenPercentage;
+	/** Overridden ScreenPercentage value */
+	float ScreenPercentage;
+
+	/** Ideal ScreenPercentage value for the HMD */
+	float IdealScreenPercentage;
+
+	/** Allows renderer to finish current frame. Setting this to 'true' may reduce the total 
+	 *  framerate (if it was above vsync) but will reduce latency. */
+	bool bAllowFinishCurrentFrame;
+
+	/** Interpupillary distance, in meters (user configurable) */
+	float InterpupillaryDistance;
+
+	/** World units (UU) to Meters scale.  Read from the level, and used to transform positional tracking data */
+	float WorldToMetersScale;
+	/** Whether world-to-meters scale is overriden or not. */
+	bool bWorldToMetersOverride; 
+
+	/** User-tunable modification to the interpupillary distance */
+	float UserDistanceToScreenModifier;
+
+	/** The FOV to render at (radians), based on the physical characteristics of the device */
+	float HFOVInRadians; // horizontal
+	float VFOVInRadians; // vertical
+
+	/** Motion prediction (in seconds). 0 - no prediction */
+	double MotionPredictionInSeconds;
+
+	/** Gain for gravity correction (should not need to be changed) */
+	float AccelGain;
+
+    /** Distortion on/off */
+    bool bHmdDistortion;
+
+	/** Chromatic aberration correction on/off */
+	bool bChromaAbCorrectionEnabled;
+
+	/** Yaw drift correction on/off */
+	bool bYawDriftCorrectionEnabled;
+
+	/** Whether or not 2D stereo settings overridden. */
+	bool bOverride2D;
+	/** HUD stereo offset */
+	float HudOffset;
+	/** Screen center adjustment for 2d elements */
+	float CanvasCenterOffset;
+
+	/** Low persistence mode */
+	bool bLowPersistenceMode;
+
+	/** Turns on/off updating view's orientation/position on a RenderThread. When it is on,
+	    latency should be significantly lower. 
+		See 'HMD UPDATEONRT ON|OFF' console command.
+	*/
+	bool				bUpdateOnRT;
+	mutable OVR::Lock	UpdateOnRTLock;
+
+	/** Overdrive brightness transitions to reduce artifacts on DK2+ displays */
+	bool bOverdrive;
+
+	/** High-quality sampling of distortion buffer for anti-aliasing */
+	bool bHQDistortion;
+
+	/** Enforces headtracking to work even in non-stereo mode (for debugging or screenshots). 
+	    See 'MOTION ENFORCE' console command. */
+	bool bHeadTrackingEnforced;
+
+	/** Is mirroring enabled or not (see 'HMD MIRROR' console cmd) */
+	bool bMirrorToWindow;
+
+#if !UE_BUILD_SHIPPING
+	/** Draw tracking camera frustum, for debugging purposes. 
+	 *  See 'HMDPOS SHOWCAMERA ON|OFF' console command.
+	 */
+	bool				bDrawTrackingCameraFrustum;
+
+	/** Turns off updating of orientation/position on game thread. See 'hmd updateongt' cmd */
+	bool				bDoNotUpdateOnGT;
+
+	/** Show status / statistics on screen. See 'hmd stats' cmd */
+	bool				bShowStats;
+
+	/** Draw lens centered grid */
+	bool				bDrawGrid;
+
+	/** Profiling mode, removed extra waits in Present (Direct Rendering). See 'hmd profile' cmd */
+	bool				bProfiling;
+#endif
+
+	/** Whether timewarp is enabled or not */
+	bool					bTimeWarp;
+
+	/** Optional far clipping plane for projection matrix */
+	float					NearClippingPlane;
+	/** Optional far clipping plane for projection matrix */
+	float					FarClippingPlane;
+
+	/** Player's orientation tracking */
+	mutable FQuat			CurHmdOrientation;
+
+	FRotator				DeltaControlRotation;    // same as DeltaControlOrientation but as rotator
+	FQuat					DeltaControlOrientation; // same as DeltaControlRotation but as quat
+
+	mutable FVector			CurHmdPosition;
+
+	mutable FQuat			LastHmdOrientation; // contains last APPLIED ON GT HMD orientation
+	FVector					LastHmdPosition;	// contains last APPLIED ON GT HMD position 
+
+	/** HMD base values, specify forward orientation and zero pos offset */
+	OVR::Vector3f			BaseOffset;      // base position, in Oculus coords
+	FQuat					BaseOrientation; // base orientation
+
+	ovrHmd					Hmd;
+	ovrEyeRenderDesc		EyeRenderDesc[2];			// 0 - left, 1 - right, same as Views
+	ovrMatrix4f				EyeProjectionMatrices[2];	// 0 - left, 1 - right, same as Views
+	ovrFovPort				EyeFov[2];					// 0 - left, 1 - right, same as Views
+	// U,V scale and offset needed for timewarp.
+	ovrRecti				EyeRenderViewport[2];		// 0 - left, 1 - right, same as Views
+	ovrSizei				TextureSize; // texture size (for both eyes)
+
+	unsigned				TrackingCaps;
+	unsigned				DistortionCaps;
+	unsigned				HmdCaps;
+
+	unsigned				SupportedTrackingCaps;
+	unsigned				SupportedDistortionCaps;
+	unsigned				SupportedHmdCaps;
+
+	FIntPoint				EyeViewportSize; // size of the viewport (for one eye). At the moment it is a half of RT.
+
+#ifndef OVR_DIRECT_RENDERING
+	struct FDistortionMesh : public OVR::RefCountBase<FDistortionMesh>
+	{
+		struct FDistortionVertex*	pVertices;
+		uint16*						pIndices;
+		unsigned					NumVertices;
+		unsigned					NumIndices;
+		unsigned					NumTriangles;
+
+		FDistortionMesh() :pVertices(NULL), pIndices(NULL), NumVertices(0), NumIndices(0), NumTriangles(0) {}
+		~FDistortionMesh() { Clear(); }
+		void Clear();
+	};
+	ovrVector2f				UVScaleOffset[2][2];	// 0 - left, 1 - right, same as Views
+	Ptr<FDistortionMesh>	pDistortionMesh[2];		// 0 - left, 1 - right, same as Views
+#else // DIRECT_RENDERING
+
+#if defined(OVR_D3D_VERSION) && (OVR_D3D_VERSION == 11)
+	TRefCountPtr<D3D11Bridge>	pD3D11Bridge;
+#endif
+#if defined(OVR_GL)
+	TRefCountPtr<OGLBridge>		pOGLBridge;
+#endif
+
+#endif // OVR_DIRECT_RENDERING
+	
+	OVR::Lock					StereoParamsLock;
+
+	// Params accessible from rendering thread. Should be filled at the beginning
+	// of the rendering thread under the StereoParamsLock.
+	struct FRenderParams
+	{
+		FVector					LastHmdPosition;	// contains last APPLIED ON GT HMD position 
+		FQuat					DeltaControlOrientation;
+		ovrPosef				EyeRenderPose[2];
+
+#ifndef OVR_DIRECT_RENDERING
+		Ptr<FDistortionMesh>	pDistortionMesh[2]; // 0 - left, 1 - right, same as Views
+		ovrVector2f				UVScale[2];			// 0 - left, 1 - right, same as Views
+		ovrVector2f				UVOffset[2];		// 0 - left, 1 - right, same as Views
+		FQuat					CurHmdOrientation;
+		FVector					CurHmdPosition;
+#else
+		ovrEyeRenderDesc		EyeRenderDesc[2];	// 0 - left, 1 - right, same as Views
+		ovrFovPort				EyeFov[2];			// 0 - left, 1 - right, same as Views
+#endif // OVR_DIRECT_RENDERING
+		bool					bFrameBegun;
+		bool					bTimeWarp;
+		FEngineShowFlags		ShowFlags; // a copy of showflags
+
+		FRenderParams(FOculusRiftHMD* plugin);
+
+		void Clear() 
+		{ 
+			#ifndef OVR_DIRECT_RENDERING
+			pDistortionMesh[0] = pDistortionMesh[1] = NULL; 
+			#endif
+		}
+	} RenderParams_RenderThread;
+
+
+	/** True, if pos tracking is enabled */
+	bool						bHmdPosTracking;
+	mutable bool				bHaveVisionTracking;
+
+	void*						OSWindowHandle;
+};
+
+DEFINE_LOG_CATEGORY_STATIC(LogHMD, Log, All);
+
+#endif //OCULUS_RIFT_SUPPORTED_PLATFORMS
diff --git a/Engine/Plugins/Runtime/OculusRift/Source/OculusRift/Private/OculusRiftPrivate.h b/Engine/Plugins/Runtime/OculusRift/Source/OculusRift/Private/OculusRiftPrivate.h
new file mode 100644
--- /dev/null
+++ b/Engine/Plugins/Runtime/OculusRift/Source/OculusRift/Private/OculusRiftPrivate.h
@@ -0,0 +1,6 @@
+// Copyright 1998-2012 Epic Games, Inc. All Rights Reserved.
+
+#pragma once
+
+#include "Engine.h"
+#include "IHeadMountedDisplay.h"
\ No newline at end of file
diff --git a/Engine/Plugins/Runtime/OculusRift/Source/OculusRift/Private/OculusRiftRender.cpp b/Engine/Plugins/Runtime/OculusRift/Source/OculusRift/Private/OculusRiftRender.cpp
new file mode 100644
--- /dev/null
+++ b/Engine/Plugins/Runtime/OculusRift/Source/OculusRift/Private/OculusRiftRender.cpp
@@ -0,0 +1,1055 @@
+// Copyright 1998-2014 Epic Games, Inc. All Rights Reserved.
+//
+#include "OculusRiftPrivate.h"
+#include "OculusRiftHMD.h"
+
+#if OCULUS_RIFT_SUPPORTED_PLATFORMS
+
+#include "../Src/OVR_Stereo.h"
+
+#include "RendererPrivate.h"
+#include "ScenePrivate.h"
+#include "PostProcess/PostProcessHMD.h"
+
+#ifndef OVR_DIRECT_RENDERING
+void FOculusRiftHMD::FDistortionMesh::Clear()
+{
+	delete[] pVertices;
+	delete[] pIndices;
+	pVertices = NULL;
+	pIndices = NULL;
+	NumVertices = NumIndices = NumTriangles = 0;
+}
+#endif
+
+FOculusRiftHMD::FRenderParams::FRenderParams(FOculusRiftHMD* plugin)
+	: 
+#ifndef OVR_DIRECT_RENDERING
+	  CurHmdOrientation(FQuat::Identity)
+	, CurHmdPosition(FVector::ZeroVector)
+	,
+#endif // OVR_DIRECT_RENDERING
+	 bFrameBegun(false)
+	, bTimeWarp(false)
+	, ShowFlags(ESFIM_All0)
+{
+	#ifndef OVR_DIRECT_RENDERING
+	FMemory::Memset(UVScale, 0, sizeof(UVScale));
+	FMemory::Memset(UVOffset, 0, sizeof(UVOffset));
+	#endif
+}
+
+void FOculusRiftHMD::PrecalculatePostProcess_NoLock()
+{
+#ifndef OVR_DIRECT_RENDERING
+	for (unsigned eyeNum = 0; eyeNum < 2; ++eyeNum)
+	{
+		// Allocate & generate distortion mesh vertices.
+		ovrDistortionMesh meshData;
+
+		if (!ovrHmd_CreateDistortionMesh(Hmd, EyeRenderDesc[eyeNum].Eye, EyeRenderDesc[eyeNum].Fov, DistortionCaps, &meshData))
+		{
+			check(false);
+			continue;
+		}
+		ovrHmd_GetRenderScaleAndOffset(EyeRenderDesc[eyeNum].Fov,
+			TextureSize, EyeRenderViewport[eyeNum],
+			(ovrVector2f*) UVScaleOffset[eyeNum]);
+
+		// alloc the data
+		ovrDistortionVertex* ovrMeshVertData = new ovrDistortionVertex[meshData.VertexCount];
+
+		// Convert to final vertex data.
+		FDistortionVertex *pVerts = new FDistortionVertex[meshData.VertexCount];
+		FDistortionVertex *pCurVert = pVerts;
+		ovrDistortionVertex* pCurOvrVert = meshData.pVertexData;
+
+		for (unsigned vertNum = 0; vertNum < meshData.VertexCount; ++vertNum)
+		{
+			pCurVert->Position.X = pCurOvrVert->ScreenPosNDC.x;
+			pCurVert->Position.Y = pCurOvrVert->ScreenPosNDC.y;
+			pCurVert->TexR = FVector2D(pCurOvrVert->TanEyeAnglesR.x, pCurOvrVert->TanEyeAnglesR.y);
+			pCurVert->TexG = FVector2D(pCurOvrVert->TanEyeAnglesG.x, pCurOvrVert->TanEyeAnglesG.y);
+			pCurVert->TexB = FVector2D(pCurOvrVert->TanEyeAnglesB.x, pCurOvrVert->TanEyeAnglesB.y);
+			pCurVert->VignetteFactor = pCurOvrVert->VignetteFactor;
+			pCurVert->TimewarpFactor = pCurOvrVert->TimeWarpFactor;
+			pCurOvrVert++;
+			pCurVert++;
+		}
+
+		pDistortionMesh[eyeNum] = *new FDistortionMesh();
+		pDistortionMesh[eyeNum]->NumTriangles = meshData.IndexCount / 3; 
+		pDistortionMesh[eyeNum]->NumIndices = meshData.IndexCount;
+		pDistortionMesh[eyeNum]->NumVertices = meshData.VertexCount;
+		pDistortionMesh[eyeNum]->pVertices = pVerts;
+
+		check(sizeof(*meshData.pIndexData) == sizeof(uint16));
+		pDistortionMesh[eyeNum]->pIndices = new uint16[meshData.IndexCount];
+		FMemory::Memcpy(pDistortionMesh[eyeNum]->pIndices, meshData.pIndexData, sizeof(uint16)*meshData.IndexCount);
+
+		ovrHmd_DestroyDistortionMesh(&meshData);
+	}
+#endif
+}
+
+void FOculusRiftHMD::DrawDistortionMesh_RenderThread(FRenderingCompositePassContext& Context, const FSceneView& View, const FIntPoint& TextureSize)
+{
+#ifndef OVR_DIRECT_RENDERING
+	float ClipSpaceQuadZ = 0.0f;
+	FMatrix QuadTexTransform = FMatrix::Identity;
+	FMatrix QuadPosTransform = FMatrix::Identity;
+	const FIntRect SrcRect = View.ViewRect;
+
+	FRHICommandListImmediate& RHICmdList = Context.RHICmdList;
+	const FSceneViewFamily& ViewFamily = *(View.Family);
+	FIntPoint ViewportSize = ViewFamily.RenderTarget->GetSizeXY();
+	RHICmdList.SetViewport(0, 0, 0.0f, ViewportSize.X, ViewportSize.Y, 1.0f);
+
+	Ptr<FDistortionMesh> mesh = RenderParams_RenderThread.pDistortionMesh[(View.StereoPass == eSSP_LEFT_EYE) ? 0 : 1];
+
+	DrawIndexedPrimitiveUP(Context.RHICmdList, PT_TriangleList, 0, mesh->NumVertices, mesh->NumTriangles, mesh->pIndices,
+		sizeof(mesh->pIndices[0]), mesh->pVertices, sizeof(mesh->pVertices[0]));
+#else
+	check(0);
+#endif
+}
+
+void FOculusRiftHMD::GetEyeRenderParams_RenderThread(EStereoscopicPass StereoPass, FVector2D& EyeToSrcUVScaleValue, FVector2D& EyeToSrcUVOffsetValue) const
+{
+	check(IsInRenderingThread());
+#ifndef OVR_DIRECT_RENDERING
+	const unsigned eyeIdx = (StereoPass == eSSP_LEFT_EYE) ? 0 : 1;
+	EyeToSrcUVOffsetValue.X = RenderParams_RenderThread.UVOffset[eyeIdx].x;
+	EyeToSrcUVOffsetValue.Y = RenderParams_RenderThread.UVOffset[eyeIdx].y;
+
+	EyeToSrcUVScaleValue.X = RenderParams_RenderThread.UVScale[eyeIdx].x;
+	EyeToSrcUVScaleValue.Y = RenderParams_RenderThread.UVScale[eyeIdx].y;
+#else
+	check(0);
+#endif
+}
+
+void FOculusRiftHMD::GetTimewarpMatrices_RenderThread(EStereoscopicPass StereoPass, FMatrix& EyeRotationStart, FMatrix& EyeRotationEnd) const
+{
+	check(IsInRenderingThread());
+#ifndef OVR_DIRECT_RENDERING
+	const ovrEyeType eye = (StereoPass == eSSP_LEFT_EYE) ? ovrEye_Left : ovrEye_Right;
+	ovrMatrix4f timeWarpMatrices[2];
+	if (RenderParams_RenderThread.bFrameBegun)
+	{
+		ovrHmd_GetEyeTimewarpMatrices(Hmd, eye, RenderParams_RenderThread.EyeRenderPose[eye], timeWarpMatrices);
+	}
+	EyeRotationStart = ToFMatrix(timeWarpMatrices[0]);
+	EyeRotationEnd = ToFMatrix(timeWarpMatrices[1]);
+#else
+	check(0);
+#endif
+}
+
+void FOculusRiftHMD::PreRenderViewFamily_RenderThread(FSceneViewFamily& ViewFamily)
+{
+	check(IsInRenderingThread());
+
+	RenderParams_RenderThread.ShowFlags = ViewFamily.EngineShowFlags;
+
+#ifndef OVR_DIRECT_RENDERING 
+	{
+		// make a copy of StereoParams to access from the RenderThread.
+		Lock::Locker lock(&StereoParamsLock);
+		RenderParams_RenderThread.UVScale[0] = UVScaleOffset[0][0];
+		RenderParams_RenderThread.UVScale[1] = UVScaleOffset[1][0];
+		RenderParams_RenderThread.UVOffset[0] = UVScaleOffset[0][1];
+		RenderParams_RenderThread.UVOffset[1] = UVScaleOffset[1][1];
+		RenderParams_RenderThread.pDistortionMesh[0] = pDistortionMesh[0];
+		RenderParams_RenderThread.pDistortionMesh[1] = pDistortionMesh[1];
+		RenderParams_RenderThread.bTimeWarp = bTimeWarp;
+	}
+#else
+	{
+		// make a copy of StereoParams to access from the RenderThread.
+		Lock::Locker lock(&StereoParamsLock);
+		RenderParams_RenderThread.EyeRenderDesc[0] = EyeRenderDesc[0];
+		RenderParams_RenderThread.EyeRenderDesc[1] = EyeRenderDesc[1];
+		RenderParams_RenderThread.EyeFov[0] = EyeFov[0];
+		RenderParams_RenderThread.EyeFov[1] = EyeFov[1];
+		RenderParams_RenderThread.bTimeWarp = bTimeWarp;
+	}
+#endif
+	// get latest orientation/position and cache it
+	if (bUpdateOnRT)
+	{
+		Lock::Locker lock(&UpdateOnRTLock);
+		RenderParams_RenderThread.LastHmdPosition = LastHmdPosition;
+		RenderParams_RenderThread.DeltaControlOrientation = DeltaControlOrientation;
+	}
+
+	BeginRendering_RenderThread();
+}
+
+void FOculusRiftHMD::PreRenderView_RenderThread(FSceneView& View)
+{
+	check(IsInRenderingThread());
+
+	if (!RenderParams_RenderThread.ShowFlags.Rendering)
+		return;
+
+	const ovrEyeType eyeIdx = (View.StereoPass == eSSP_LEFT_EYE) ? ovrEye_Left : ovrEye_Right;
+	FQuat	CurrentHmdOrientation;
+	FVector	CurrentHmdPosition;
+
+	// Get new predicted pose to corresponding eye.
+	ovrPosef eyeRenderPose = ovrHmd_GetEyePose(Hmd, eyeIdx);
+	PoseToOrientationAndPosition(eyeRenderPose, CurrentHmdOrientation, CurrentHmdPosition);
+	RenderParams_RenderThread.EyeRenderPose[eyeIdx] = eyeRenderPose;
+
+	if (bUpdateOnRT)
+	{
+		// Apply updated pose to corresponding View at recalc matrices.
+		UpdatePlayerViewPoint(CurrentHmdOrientation, CurrentHmdPosition, 
+			RenderParams_RenderThread.LastHmdPosition, RenderParams_RenderThread.DeltaControlOrientation,
+			View.BaseHmdOrientation, View.BaseHmdLocation,
+			View.ViewRotation, View.ViewLocation);
+		View.UpdateViewMatrix();
+	}
+}
+
+void FOculusRiftHMD::BeginRendering_RenderThread()
+{
+	check(IsInRenderingThread());
+	if (RenderParams_RenderThread.bFrameBegun)
+	{
+		return;
+	}
+
+#ifdef OVR_DIRECT_RENDERING 
+	{
+		Lock::Locker lock(&StereoParamsLock);
+
+		GetActiveRHIBridgeImpl()->BeginRendering();
+	}
+
+	ovrHmd_BeginFrame(Hmd, 0);
+#else
+	ovrHmd_BeginFrameTiming(Hmd, 0);
+#endif
+	RenderParams_RenderThread.bFrameBegun = true;
+}
+
+
+#ifdef OVR_DIRECT_RENDERING 
+FOculusRiftHMD::BridgeBaseImpl* FOculusRiftHMD::GetActiveRHIBridgeImpl()
+{
+#if defined(OVR_D3D_VERSION) && (OVR_D3D_VERSION == 11)
+	if (pD3D11Bridge)
+	{
+		return pD3D11Bridge;
+	}
+#endif
+#if defined(OVR_GL)
+	if (pOGLBridge)
+	{
+		return pOGLBridge;
+	}
+#endif
+	return nullptr;
+}
+
+void FOculusRiftHMD::CalculateRenderTargetSize(uint32& InOutSizeX, uint32& InOutSizeY) const
+{
+	check(IsInGameThread());
+
+	static const auto CVar = IConsoleManager::Get().FindTConsoleVariableDataFloat(TEXT("r.ScreenPercentage"));
+	float value = CVar->GetValueOnGameThread();
+	if (value > 0.0f)
+	{
+		InOutSizeX = FMath::CeilToInt(InOutSizeX * value/100.f);
+		InOutSizeY = FMath::CeilToInt(InOutSizeY * value/100.f);
+	}
+}
+
+bool FOculusRiftHMD::NeedReAllocateViewportRenderTarget(const FViewport& Viewport) const
+{
+	check(IsInGameThread());
+	if (IsStereoEnabled())
+	{
+		const uint32 InSizeX = Viewport.GetSizeXY().X;
+		const uint32 InSizeY = Viewport.GetSizeXY().Y;
+		FIntPoint RenderTargetSize;
+		RenderTargetSize.X = Viewport.GetRenderTargetTexture()->GetSizeX();
+		RenderTargetSize.Y = Viewport.GetRenderTargetTexture()->GetSizeY();
+
+		uint32 NewSizeX = InSizeX, NewSizeY = InSizeY;
+		CalculateRenderTargetSize(NewSizeX, NewSizeY);
+		if (NewSizeX != RenderTargetSize.X || NewSizeY != RenderTargetSize.Y)
+		{
+			return true;
+		}
+	}
+	return false;
+}
+
+#else // no direct rendering
+
+void FOculusRiftHMD::FinishRenderingFrame_RenderThread(FRHICommandListImmediate& RHICmdList)
+{
+	if (RenderParams_RenderThread.bFrameBegun)
+	{
+		check(IsInRenderingThread());
+		if (RenderParams_RenderThread.bTimeWarp)
+		{
+			RHICmdList.BlockUntilGPUIdle();
+		}
+		ovrHmd_EndFrameTiming(Hmd);
+		RenderParams_RenderThread.bFrameBegun = false;
+	}
+}
+
+#endif // #ifdef OVR_DIRECT_RENDERING 
+
+static const char* FormatLatencyReading(char* buff, UPInt size, float val)
+{
+	if (val < 0.000001f)
+	{
+		OVR_strcpy(buff, size, "N/A   ");
+	}
+	else
+	{
+		OVR_sprintf(buff, size, "%4.2fms", val * 1000.0f);
+	}
+	return buff;
+}
+
+#if !UE_BUILD_SHIPPING
+static void RenderLines(FCanvas* Canvas, int numLines, const FColor& c, float* x, float* y)
+{
+	for (int i = 0; i < numLines; ++i)
+	{
+		FCanvasLineItem line(FVector2D(x[i*2], y[i*2]), FVector2D(x[i*2+1], y[i*2+1]));
+		line.SetColor(FLinearColor(c));
+		Canvas->DrawItem(line);
+	}
+}
+#endif // #if !UE_BUILD_SHIPPING
+
+void FOculusRiftHMD::DrawDebug(UCanvas* Canvas, EStereoscopicPass StereoPass)
+{
+#if !UE_BUILD_SHIPPING
+	check(IsInGameThread());
+	if (StereoPass == eSSP_FULL)
+	{
+		if (bDrawGrid)
+		{
+			const FColor cNormal(255, 0, 0);
+			const FColor cSpacer(255, 255, 0);
+			const FColor cMid(0, 128, 255);
+			for (int eye = 0; eye < 2; ++eye)
+			{
+				int lineStep = 1;
+				int midX = 0;
+				int midY = 0;
+				int limitX = 0;
+				int limitY = 0;
+
+				int renderViewportX = EyeRenderDesc[eye].DistortedViewport.Pos.x;
+				int renderViewportY = EyeRenderDesc[eye].DistortedViewport.Pos.y;
+				int renderViewportW = EyeRenderDesc[eye].DistortedViewport.Size.w;
+				int renderViewportH = EyeRenderDesc[eye].DistortedViewport.Size.h;
+
+				lineStep = 48;
+				Vector2f rendertargetNDC = OVR::FovPort(EyeRenderDesc[eye].Fov).TanAngleToRendertargetNDC(Vector2f(0.0f));
+				midX = (int)((rendertargetNDC.x * 0.5f + 0.5f) * (float)renderViewportW + 0.5f);
+				midY = (int)((rendertargetNDC.y * 0.5f + 0.5f) * (float)renderViewportH + 0.5f);
+				limitX = Alg::Max(renderViewportW - midX, midX);
+				limitY = Alg::Max(renderViewportH - midY, midY);
+
+				int spacerMask = (lineStep << 2) - 1;
+
+				for (int xp = 0; xp < limitX; xp += lineStep)
+				{
+					float x[4];
+					float y[4];
+					x[0] = (float)(midX + xp) + renderViewportX;
+					y[0] = (float)0 + renderViewportY;
+					x[1] = (float)(midX + xp) + renderViewportX;
+					y[1] = (float)renderViewportH + renderViewportY;
+					x[2] = (float)(midX - xp) + renderViewportX;
+					y[2] = (float)0 + renderViewportY;
+					x[3] = (float)(midX - xp) + renderViewportX;
+					y[3] = (float)renderViewportH + renderViewportY;
+					if (xp == 0)
+					{
+						RenderLines(Canvas->Canvas, 1, cMid, x, y);
+					}
+					else if ((xp & spacerMask) == 0)
+					{
+						RenderLines(Canvas->Canvas, 2, cSpacer, x, y);
+					}
+					else
+					{
+						RenderLines(Canvas->Canvas, 2, cNormal, x, y);
+					}
+				}
+				for (int yp = 0; yp < limitY; yp += lineStep)
+				{
+					float x[4];
+					float y[4];
+					x[0] = (float)0 + renderViewportX;
+					y[0] = (float)(midY + yp) + renderViewportY;
+					x[1] = (float)renderViewportW + renderViewportX;
+					y[1] = (float)(midY + yp) + renderViewportY;
+					x[2] = (float)0 + renderViewportX;
+					y[2] = (float)(midY - yp) + renderViewportY;
+					x[3] = (float)renderViewportW + renderViewportX;
+					y[3] = (float)(midY - yp) + renderViewportY;
+					if (yp == 0)
+					{
+						RenderLines(Canvas->Canvas, 1, cMid, x, y);
+					}
+					else if ((yp & spacerMask) == 0)
+					{
+						RenderLines(Canvas->Canvas, 2, cSpacer, x, y);
+					}
+					else
+					{
+						RenderLines(Canvas->Canvas, 2, cNormal, x, y);
+					}
+				}
+			}
+		}
+		return;
+	}
+	else if (IsStereoEnabled() && bShowStats)
+	{
+		static const FColor TextColor(0,255,0);
+		// Pick a larger font on console.
+		UFont* const Font = FPlatformProperties::SupportsWindowedMode() ? GEngine->GetSmallFont() : GEngine->GetMediumFont();
+		const int32 RowHeight = FMath::TruncToInt(Font->GetMaxCharHeight() * 1.1f);
+
+		float ClipX = Canvas->ClipX;
+		float ClipY = Canvas->ClipY;
+		float LeftPos = 0;
+
+		ClipX -= 100;
+		//ClipY = ClipY * 0.60;
+		LeftPos = ClipX * 0.3f;
+		float TopPos = ClipY * 0.4f;
+
+		int32 X = (int32)LeftPos;
+		int32 Y = (int32)TopPos;
+
+		FString Str, StatusStr;
+		// First row
+		Str = FString::Printf(TEXT("TimeWarp: %s"), (bTimeWarp) ? TEXT("ON") : TEXT("OFF"));
+		Canvas->Canvas->DrawShadowedString(X, Y, *Str, Font, TextColor);
+
+		Y += RowHeight;
+		Str = FString::Printf(TEXT("VSync: %s"), (bVSync) ? TEXT("ON") : TEXT("OFF"));
+		Canvas->Canvas->DrawShadowedString(X, Y, *Str, Font, TextColor);
+		
+		Y += RowHeight;
+		Str = FString::Printf(TEXT("Upd on GT/RT: %s / %s"), (!bDoNotUpdateOnGT) ? TEXT("ON") : TEXT("OFF"), 
+			(bUpdateOnRT) ? TEXT("ON") : TEXT("OFF"));
+		Canvas->Canvas->DrawShadowedString(X, Y, *Str, Font, TextColor);
+
+		Y += RowHeight;
+		static IConsoleVariable* CFinishFrameVar = IConsoleManager::Get().FindConsoleVariable(TEXT("r.FinishCurrentFrame"));
+		int finFr = CFinishFrameVar->GetInt();
+		Str = FString::Printf(TEXT("FinFr: %s"), (finFr || bTimeWarp) ? TEXT("ON") : TEXT("OFF"));
+		Canvas->Canvas->DrawShadowedString(X, Y, *Str, Font, TextColor);
+
+		Y += RowHeight;
+		static IConsoleVariable* CScrPercVar = IConsoleManager::Get().FindConsoleVariable(TEXT("r.ScreenPercentage"));
+		int32 sp = (int32)CScrPercVar->GetFloat();
+		Str = FString::Printf(TEXT("SP: %d"), sp);
+		Canvas->Canvas->DrawShadowedString(X, Y, *Str, Font, TextColor);
+
+		Y += RowHeight;
+		Str = FString::Printf(TEXT("FOV V/H: %.2f / %.2f deg"), FMath::RadiansToDegrees(VFOVInRadians), FMath::RadiansToDegrees(HFOVInRadians));
+		Canvas->Canvas->DrawShadowedString(X, Y, *Str, Font, TextColor);
+
+		Y += RowHeight;
+		Str = FString::Printf(TEXT("W-to-m scale: %.2f uu/m"), WorldToMetersScale);
+		Canvas->Canvas->DrawShadowedString(X, Y, *Str, Font, TextColor);
+
+		if ((SupportedHmdCaps & ovrHmdCap_DynamicPrediction) != 0)
+		{
+			float latencies[3] = { 0.0f, 0.0f, 0.0f };
+			if (ovrHmd_GetFloatArray(Hmd, "DK2Latency", latencies, 3) == 3)
+			{
+				Y += RowHeight;
+
+				char buf[3][20];
+				char destStr[100];
+
+				OVR_sprintf(destStr, sizeof(destStr), "Latency, ren: %s tw: %s pp: %s",
+					FormatLatencyReading(buf[0], sizeof(buf[0]), latencies[0]),
+					FormatLatencyReading(buf[1], sizeof(buf[1]), latencies[1]),
+					FormatLatencyReading(buf[2], sizeof(buf[2]), latencies[2]));
+
+				Str = ANSI_TO_TCHAR(destStr);
+				Canvas->Canvas->DrawShadowedString(X, Y, *Str, Font, TextColor);
+			}
+		}
+
+		// Second row
+		X = (int32)LeftPos + 200;
+		Y = (int32)TopPos;
+
+		StatusStr = ((SupportedTrackingCaps & ovrTrackingCap_Position) != 0) ?
+			((bHmdPosTracking) ? TEXT("ON") : TEXT("OFF")) : TEXT("UNSUP");
+		Str = FString::Printf(TEXT("PosTr: %s"), *StatusStr);
+		Canvas->Canvas->DrawShadowedString(X, Y, *Str, Font, TextColor);
+		Y += RowHeight;
+
+		Str = FString::Printf(TEXT("Vision: %s"), (bHaveVisionTracking) ? TEXT("ACQ") : TEXT("LOST"));
+		Canvas->Canvas->DrawShadowedString(X, Y, *Str, Font, TextColor);
+		Y += RowHeight;
+
+		Str = FString::Printf(TEXT("IPD: %.2f mm"), InterpupillaryDistance*1000.f);
+		Canvas->Canvas->DrawShadowedString(X, Y, *Str, Font, TextColor);
+		Y += RowHeight;
+
+		StatusStr = ((SupportedHmdCaps & ovrHmdCap_LowPersistence) != 0) ? 
+			((bLowPersistenceMode) ? TEXT("ON") : TEXT("OFF")) : TEXT("UNSUP");
+		Str = FString::Printf(TEXT("LowPers: %s"), *StatusStr);
+		Canvas->Canvas->DrawShadowedString(X, Y, *Str, Font, TextColor);
+		Y += RowHeight;
+
+		StatusStr = ((SupportedDistortionCaps & ovrDistortionCap_Overdrive) != 0) ?
+			((bOverdrive) ? TEXT("ON") : TEXT("OFF")) : TEXT("UNSUP");
+		Str = FString::Printf(TEXT("Overdrive: %s"), *StatusStr);
+		Canvas->Canvas->DrawShadowedString(X, Y, *Str, Font, TextColor);
+		Y += RowHeight;
+	}
+#endif // #if !UE_BUILD_SHIPPING
+}
+
+void FOculusRiftHMD::UpdateViewport(bool bUseSeparateRenderTarget, const FViewport& InViewport)
+{
+	check(IsInGameThread());
+
+	FRHIViewport* const ViewportRHI = InViewport.GetViewportRHI().GetReference();
+
+	if (!IsStereoEnabled())
+	{
+		if (!bUseSeparateRenderTarget)
+		{
+			ViewportRHI->SetCustomPresent(nullptr);
+		}
+#if PLATFORM_WINDOWS
+		if (OSWindowHandle)
+		{
+			ovrHmd_AttachToWindow(Hmd, NULL, NULL, NULL);
+			OSWindowHandle = nullptr;
+		}
+#endif
+		return;
+	}
+
+#if PLATFORM_WINDOWS
+	void *wnd = ViewportRHI->GetNativeWindow();
+	if (wnd && wnd != OSWindowHandle)
+	{
+		OSWindowHandle = wnd;
+		HWND Window = *(HWND*)wnd;
+		ovrHmd_AttachToWindow(Hmd, Window, NULL, NULL);
+	}
+#endif
+
+#ifdef OVR_DIRECT_RENDERING
+
+	check(GetActiveRHIBridgeImpl());
+
+	const FTexture2DRHIRef& RT = InViewport.GetRenderTargetTexture();
+	check(IsValidRef(RT));
+	const FIntPoint NewEyeRTSize = FIntPoint((RT->GetSizeX() + 1) / 2, RT->GetSizeY());
+	if (EyeViewportSize != NewEyeRTSize)
+	{
+		EyeViewportSize.X = NewEyeRTSize.X;
+		EyeViewportSize.Y = NewEyeRTSize.Y;
+		bNeedUpdateStereoRenderingParams = true;
+	}
+	if (bNeedUpdateStereoRenderingParams)
+	{
+		UpdateStereoRenderingParams();
+	}
+	GetActiveRHIBridgeImpl()->UpdateViewport(InViewport, ViewportRHI);
+#endif // #ifdef OVR_DIRECT_RENDERING
+}
+
+#ifdef OVR_DIRECT_RENDERING
+void FOculusRiftHMD::ShutdownRendering()
+{
+#if defined(OVR_D3D_VERSION) && (OVR_D3D_VERSION == 11)
+	if (pD3D11Bridge)
+	{
+		pD3D11Bridge->Shutdown();
+		pD3D11Bridge = NULL;
+	}
+#endif
+#if defined(OVR_GL)
+	if (pOGLBridge)
+	{
+		pOGLBridge->Shutdown();
+		pOGLBridge = NULL;
+	}
+#endif
+
+	ovrHmd_Destroy(Hmd);
+	Hmd = nullptr;
+}
+
+#if defined(OVR_D3D_VERSION) && (OVR_D3D_VERSION == 11)
+FOculusRiftHMD::D3D11Bridge::D3D11Bridge(FOculusRiftHMD* plugin):
+	BridgeBaseImpl(plugin), 
+	bNeedReinitEyeTextures(false)
+{
+	memset(&Cfg, 0, sizeof(Cfg));
+	memset(&EyeTexture, 0, sizeof(EyeTexture));
+	memset(&EyeTexture_RenderThread, 0, sizeof(EyeTexture_RenderThread));
+}
+
+void FOculusRiftHMD::D3D11Bridge::BeginRendering()
+{
+	check(IsInRenderingThread());
+	if (bInitialized)
+	{
+		if (bNeedReinitRendererAPI)
+		{
+			OVR::Lock::Locker lock(&ModifyLock);
+			check(Cfg.D3D11.pSwapChain); // make sure Config is initialized
+			Plugin->UpdateDistortionCaps();
+			if (!ovrHmd_ConfigureRendering(Plugin->Hmd, &Cfg.Config, Plugin->DistortionCaps, 
+				Plugin->RenderParams_RenderThread.EyeFov, Plugin->RenderParams_RenderThread.EyeRenderDesc))
+			{
+				UE_LOG(LogHMD, Warning, TEXT("D3D11 ovrHmd_ConfigureRenderAPI failed."));
+				return;
+			}
+			bNeedReinitRendererAPI = false;
+		}
+
+		if (bNeedReinitEyeTextures)
+		{
+			OVR::Lock::Locker lock(&ModifyEyeTexturesLock);
+
+			ovrD3D11TextureData oldEye0 = EyeTexture_RenderThread[0].D3D11;
+			ovrD3D11TextureData oldEye1 = EyeTexture_RenderThread[1].D3D11;
+
+			EyeTexture_RenderThread[0] = EyeTexture[0];
+			if (EyeTexture_RenderThread[0].D3D11.pTexture)
+			{
+				EyeTexture_RenderThread[0].D3D11.pTexture->AddRef();
+			}
+			if (EyeTexture_RenderThread[0].D3D11.pSRView)
+			{
+				EyeTexture_RenderThread[0].D3D11.pSRView->AddRef();
+			}
+			EyeTexture_RenderThread[1] = EyeTexture[1];
+			if (EyeTexture_RenderThread[1].D3D11.pTexture)
+			{
+				EyeTexture_RenderThread[1].D3D11.pTexture->AddRef();
+			}
+			if (EyeTexture_RenderThread[1].D3D11.pSRView)
+			{
+				EyeTexture_RenderThread[1].D3D11.pSRView->AddRef();
+			}
+
+			if (oldEye0.pTexture)
+			{
+				oldEye0.pTexture->Release();
+			}
+			if (oldEye1.pTexture)
+			{
+				oldEye1.pTexture->Release();
+			}
+			if (oldEye0.pSRView)
+			{
+				oldEye0.pSRView->Release();
+			}
+			if (oldEye1.pSRView)
+			{
+				oldEye1.pSRView->Release();
+			}
+
+			bNeedReinitEyeTextures = false;
+		}
+	}
+}
+
+void FOculusRiftHMD::D3D11Bridge::FinishRendering()
+{
+	check(IsInRenderingThread());
+	if (Plugin->RenderParams_RenderThread.bFrameBegun)
+	{
+		// Finish the frame and let OVR do buffer swap (Present) and flush/sync.
+		const ovrTexture eyeTextures[2] = { EyeTexture_RenderThread[0].Texture, EyeTexture_RenderThread[1].Texture };
+		ovrHmd_EndFrame(Plugin->Hmd, Plugin->RenderParams_RenderThread.EyeRenderPose, eyeTextures); // This function will present
+		Plugin->RenderParams_RenderThread.bFrameBegun = false;
+	}
+	else
+	{
+		UE_LOG(LogHMD, Warning, TEXT("Skipping frame: FinishRendering called with no corresponding BeginRendering (was BackBuffer re-allocated?)"));
+	}
+}
+
+void FOculusRiftHMD::D3D11Bridge::Reset_RenderThread()
+{
+	Cfg.D3D11.pDevice = nullptr;
+	Cfg.D3D11.pDeviceContext = nullptr;
+
+	OVR::Lock::Locker lock(&ModifyEyeTexturesLock);
+	if (EyeTexture[0].D3D11.pTexture)
+	{
+		EyeTexture[0].D3D11.pTexture->Release();
+		EyeTexture[0].D3D11.pTexture = (ID3D11Texture2D*)nullptr;
+	}
+	if (EyeTexture[0].D3D11.pSRView)
+	{
+		EyeTexture[0].D3D11.pSRView->Release();
+		EyeTexture[0].D3D11.pSRView = nullptr;
+	}
+	if (EyeTexture[1].D3D11.pTexture)
+	{
+		EyeTexture[1].D3D11.pTexture->Release();
+		EyeTexture[1].D3D11.pTexture = nullptr;
+	}
+	if (EyeTexture[1].D3D11.pSRView)
+	{
+		EyeTexture[1].D3D11.pSRView->Release();
+		EyeTexture[1].D3D11.pSRView = nullptr;
+	}
+
+	Cfg.D3D11.pBackBufferRT = nullptr;
+	Cfg.D3D11.pSwapChain = nullptr;
+
+	if (EyeTexture_RenderThread[0].D3D11.pTexture)
+	{
+		EyeTexture_RenderThread[0].D3D11.pTexture->Release();
+		EyeTexture_RenderThread[0].D3D11.pTexture = nullptr;
+	}
+	if (EyeTexture_RenderThread[0].D3D11.pSRView)
+	{
+		EyeTexture_RenderThread[0].D3D11.pSRView->Release();
+		EyeTexture_RenderThread[0].D3D11.pSRView = nullptr;
+	}
+	if (EyeTexture_RenderThread[1].D3D11.pTexture)
+	{
+		EyeTexture_RenderThread[1].D3D11.pTexture->Release();
+		EyeTexture_RenderThread[1].D3D11.pTexture = nullptr;
+	}
+	if (EyeTexture_RenderThread[1].D3D11.pSRView)
+	{
+		EyeTexture_RenderThread[1].D3D11.pSRView->Release();
+		EyeTexture_RenderThread[1].D3D11.pSRView = nullptr;
+	}
+	bNeedReinitEyeTextures = false;
+	bNeedReinitRendererAPI = false;
+}
+
+void FOculusRiftHMD::D3D11Bridge::Reset()
+{
+	if (IsInGameThread())
+	{
+		ENQUEUE_UNIQUE_RENDER_COMMAND_ONEPARAMETER(ResetD3D,
+		FOculusRiftHMD::D3D11Bridge*, Bridge, this,
+		{
+			Bridge->Reset_RenderThread();
+		});
+		// Wait for all resources to be released
+		FlushRenderingCommands();
+	}
+	else
+	{
+		Reset_RenderThread();
+	}
+
+	bInitialized = false;
+}
+
+void FOculusRiftHMD::D3D11Bridge::UpdateViewport(const FViewport& Viewport, FRHIViewport* ViewportRHI)
+{
+	check(IsInGameThread());
+	check(ViewportRHI);
+
+	const FTexture2DRHIRef& RT = Viewport.GetRenderTargetTexture();
+	ID3D11RenderTargetView* const	pD3DBBRT = (ID3D11RenderTargetView*)ViewportRHI->GetNativeBackBufferRT();
+	IDXGISwapChain*         const	pD3DSC = (IDXGISwapChain*)ViewportRHI->GetNativeSwapChain();
+	check(IsValidRef(RT));
+	ID3D11Texture2D* const 			pD3DRT = (ID3D11Texture2D*)RT->GetNativeResource();
+	ID3D11ShaderResourceView* const pD3DSRV = (ID3D11ShaderResourceView*)RT->GetNativeShaderResourceView();
+	const uint32 RTSizeX = RT->GetSizeX();
+	const uint32 RTSizeY = RT->GetSizeY();
+
+	ID3D11Device* D3DDevice = (ID3D11Device*)RHIGetNativeDevice();
+	ID3D11DeviceContext* D3DDeviceContext = nullptr;
+	if (D3DDevice)
+	{
+		D3DDevice->GetImmediateContext(&D3DDeviceContext);
+	}
+	if (D3DDevice != Cfg.D3D11.pDevice || D3DDeviceContext != Cfg.D3D11.pDeviceContext)
+	{
+		OVR::Lock::Locker lock(&ModifyLock);
+		Cfg.D3D11.Header.API = ovrRenderAPI_D3D11;
+		Cfg.D3D11.Header.Multisample = 1; //?? RenderParams.Multisample;
+		// Note, neither Device nor Context are AddRef-ed here. Not sure, if we need to.
+		Cfg.D3D11.pDevice = D3DDevice;
+		Cfg.D3D11.pDeviceContext = D3DDeviceContext;
+		bNeedReinitRendererAPI = true;
+		Plugin->RenderParams_RenderThread.bFrameBegun = false;
+		bInitialized = true;
+	}
+
+	if (Cfg.D3D11.pBackBufferRT != pD3DBBRT ||
+		Cfg.D3D11.pSwapChain != pD3DSC ||
+		Cfg.D3D11.Header.RTSize.w != Viewport.GetSizeXY().X ||
+		Cfg.D3D11.Header.RTSize.h != Viewport.GetSizeXY().Y)
+	{
+		OVR::Lock::Locker lock(&ModifyLock);
+		// Note, neither BackBufferRT nor SwapChain are AddRef-ed here. Not sure, if we need to.
+		// If yes, then them should be released in ReleaseBackBuffer().
+		Cfg.D3D11.pBackBufferRT = pD3DBBRT;
+		Cfg.D3D11.pSwapChain = pD3DSC;
+		Cfg.D3D11.Header.RTSize.w = Viewport.GetSizeXY().X;
+		Cfg.D3D11.Header.RTSize.h = Viewport.GetSizeXY().Y;
+		bNeedReinitRendererAPI = true;
+		Plugin->RenderParams_RenderThread.bFrameBegun = false;
+	}
+
+	if (EyeTexture[0].D3D11.pTexture != pD3DRT || EyeTexture[0].D3D11.pSRView != pD3DSRV ||
+		EyeTexture[0].D3D11.Header.TextureSize.w != RTSizeX || EyeTexture[0].D3D11.Header.TextureSize.h != RTSizeY)
+	{
+		OVR::Lock::Locker lock(&ModifyEyeTexturesLock);
+
+		ovrD3D11TextureData oldEye0 = EyeTexture[0].D3D11;
+		EyeTexture[0].D3D11.Header.API = ovrRenderAPI_D3D11;
+		EyeTexture[0].D3D11.Header.TextureSize = Sizei(RTSizeX, RTSizeY);
+		EyeTexture[0].D3D11.Header.RenderViewport = Plugin->EyeRenderViewport[0];
+		EyeTexture[0].D3D11.pTexture = pD3DRT;
+		EyeTexture[0].D3D11.pSRView = pD3DSRV;
+		if (EyeTexture[0].D3D11.pTexture)
+		{
+			EyeTexture[0].D3D11.pTexture->AddRef();
+		}
+		if (EyeTexture[0].D3D11.pSRView)
+		{
+			EyeTexture[0].D3D11.pSRView->AddRef();
+		}
+
+		if (oldEye0.pTexture)
+		{
+			oldEye0.pTexture->Release();
+		}
+		if (oldEye0.pSRView)
+		{
+			oldEye0.pSRView->Release();
+		}
+
+		// Right eye uses the same texture, but different rendering viewport.
+		ovrD3D11TextureData oldEye1 = EyeTexture[1].D3D11;
+		EyeTexture[1] = EyeTexture[0];
+		EyeTexture[1].D3D11.Header.RenderViewport = Plugin->EyeRenderViewport[1];
+		if (EyeTexture[1].D3D11.pTexture)
+		{
+			EyeTexture[1].D3D11.pTexture->AddRef();
+		}
+		if (EyeTexture[1].D3D11.pSRView)
+		{
+			EyeTexture[1].D3D11.pSRView->AddRef();
+		}
+
+		if (oldEye1.pTexture)
+		{
+			oldEye1.pTexture->Release();
+		}
+		if (oldEye1.pSRView)
+		{
+			oldEye1.pSRView->Release();
+		}
+
+		bNeedReinitEyeTextures = true;
+	}
+
+	this->ViewportRHI = ViewportRHI;
+ 	ViewportRHI->SetCustomPresent(this);
+}
+
+
+void FOculusRiftHMD::D3D11Bridge::OnBackBufferResize()
+{
+	Cfg.D3D11.pBackBufferRT		= nullptr;
+	Cfg.D3D11.pSwapChain		= nullptr;
+ 
+ 	bNeedReinitRendererAPI = true;
+
+	// if we are in the middle of rendering: prevent from calling EndFrame
+	Plugin->RenderParams_RenderThread.bFrameBegun = false;
+}
+
+bool FOculusRiftHMD::D3D11Bridge::Present(int SyncInterval)
+{
+	check(IsInRenderingThread());
+
+	FinishRendering();
+
+	return false; // indicates that we are presenting here, UE shouldn't do Present.
+}
+#endif // #if defined(OVR_D3D_VERSION) && (OVR_D3D_VERSION == 11)
+
+//////////////////////////////////////////////////////////////////////////
+#if defined(OVR_GL)
+FOculusRiftHMD::OGLBridge::OGLBridge(FOculusRiftHMD* plugin) : 
+	BridgeBaseImpl(plugin), 
+	bNeedReinitEyeTextures(false)
+{
+	memset(&Cfg, 0, sizeof(Cfg));
+	memset(&EyeTexture, 0, sizeof(EyeTexture));
+	memset(&EyeTexture_RenderThread, 0, sizeof(EyeTexture_RenderThread));
+	Init();
+}
+
+void FOculusRiftHMD::OGLBridge::BeginRendering()
+{
+	if (bInitialized)
+	{
+		if (bNeedReinitRendererAPI)
+		{
+			OVR::Lock::Locker lock(&ModifyLock);
+			Plugin->UpdateDistortionCaps();
+			if (!ovrHmd_ConfigureRendering(Plugin->Hmd, &Cfg.Config, Plugin->DistortionCaps, 
+				Plugin->RenderParams_RenderThread.EyeFov, Plugin->RenderParams_RenderThread.EyeRenderDesc))
+			{
+				UE_LOG(LogHMD, Warning, TEXT("OGL ovrHmd_ConfigureRenderAPI failed."));
+				return;
+			}
+			bNeedReinitRendererAPI = false;
+		}
+		if (bNeedReinitEyeTextures)
+		{
+			OVR::Lock::Locker lock(&ModifyEyeTexturesLock);
+
+			EyeTexture_RenderThread[0] = EyeTexture[0];
+			EyeTexture_RenderThread[1] = EyeTexture[1];
+			bNeedReinitEyeTextures = false;
+		}
+	}
+}
+
+void FOculusRiftHMD::OGLBridge::FinishRendering()
+{
+	check(IsInRenderingThread());
+
+	if (Plugin->RenderParams_RenderThread.bFrameBegun)
+	{
+		// Finish the frame and let OVR do buffer swap (Present) and flush/sync.
+		const ovrTexture eyeTextures[2] = { EyeTexture_RenderThread[0].Texture, EyeTexture_RenderThread[1].Texture };
+		ovrHmd_EndFrame(Plugin->Hmd, Plugin->RenderParams_RenderThread.EyeRenderPose, eyeTextures); // This function will present
+		Plugin->RenderParams_RenderThread.bFrameBegun = false;
+	}
+	else
+	{
+		UE_LOG(LogHMD, Warning, TEXT("Skipping frame: FinishRendering called with no corresponding BeginRendering (was BackBuffer re-allocated?)"));
+	}
+}
+
+void FOculusRiftHMD::OGLBridge::Init()
+{
+	Cfg.OGL.Header.API = ovrRenderAPI_OpenGL;
+	Cfg.OGL.Header.Multisample = 1; 
+	bNeedReinitRendererAPI = true;
+	bInitialized = true;
+}
+
+void FOculusRiftHMD::OGLBridge::Reset()
+{
+	EyeTexture[0].OGL.TexId = 0;
+	EyeTexture[1].OGL.TexId = 0;
+	EyeTexture_RenderThread[0].OGL.TexId = 0;
+	EyeTexture_RenderThread[1].OGL.TexId = 0;
+
+	Plugin->RenderParams_RenderThread.bFrameBegun = false;
+	bNeedReinitEyeTextures = false;
+	bNeedReinitRendererAPI = false;
+	bInitialized = false;
+}
+
+void FOculusRiftHMD::OGLBridge::OnBackBufferResize()
+{
+	bNeedReinitRendererAPI = true;
+
+	// if we are in the middle of rendering: prevent from calling EndFrame
+	Plugin->RenderParams_RenderThread.bFrameBegun = false;
+}
+
+void FOculusRiftHMD::OGLBridge::UpdateViewport(const FViewport& Viewport, FRHIViewport* ViewportRHI)
+{
+	check(IsInGameThread());
+	check(ViewportRHI);
+
+	bool bWinChanged = false;
+#if PLATFORM_WINDOWS
+	HWND Window = *(HWND*)ViewportRHI->GetNativeWindow();
+	bWinChanged = (Cfg.OGL.Window != Window);
+#elif PLATFORM_MAC
+	//@TODO
+	//NSWindow
+#elif PLATFORM_LINUX
+	//@TODO
+#endif
+	if (bWinChanged ||
+		Cfg.OGL.Header.RTSize.w != Viewport.GetSizeXY().X ||
+		Cfg.OGL.Header.RTSize.h != Viewport.GetSizeXY().Y)
+	{
+		OVR::Lock::Locker lock(&ModifyLock);
+
+		Cfg.OGL.Header.RTSize = Sizei(Viewport.GetSizeXY().X, Viewport.GetSizeXY().Y);
+#if PLATFORM_WINDOWS
+		Cfg.OGL.Window = Window;
+#elif PLATFORM_MAC
+		//@TODO
+#elif PLATFORM_LINUX
+		//@TODO
+		//	Cfg.OGL.Disp = AddParams; //?
+		//	Cfg.OGL.Win = hWnd; //?
+#endif
+		bNeedReinitRendererAPI = true;
+	}
+
+	const FTexture2DRHIRef& RT = Viewport.GetRenderTargetTexture();
+	check(IsValidRef(RT));
+	const uint32 RTSizeX = RT->GetSizeX();
+	const uint32 RTSizeY = RT->GetSizeY();
+	GLuint RTTexId = *(GLuint*)RT->GetNativeResource();
+
+	if (EyeTexture[0].OGL.TexId != RTTexId ||
+		EyeTexture[0].OGL.Header.TextureSize.w != RTSizeX || EyeTexture[0].OGL.Header.TextureSize.h != RTSizeY)
+	{
+		OVR::Lock::Locker lock(&ModifyEyeTexturesLock);
+
+		EyeTexture[0].OGL.Header.API = ovrRenderAPI_OpenGL;
+		EyeTexture[0].OGL.Header.TextureSize = Sizei(RTSizeX, RTSizeY);
+		EyeTexture[0].OGL.Header.RenderViewport = Plugin->EyeRenderViewport[0];
+		EyeTexture[0].OGL.TexId = RTTexId;
+
+		// Right eye uses the same texture, but different rendering viewport.
+		EyeTexture[1] = EyeTexture[0];
+		EyeTexture[1].OGL.Header.RenderViewport = Plugin->EyeRenderViewport[1];
+
+		bNeedReinitEyeTextures = true;
+	}
+	this->ViewportRHI = ViewportRHI;
+	ViewportRHI->SetCustomPresent(this);
+}
+
+bool FOculusRiftHMD::OGLBridge::Present(int SyncInterval)
+{
+	check(IsInRenderingThread());
+
+	FinishRendering();
+
+	return false; // indicates that we are presenting here, UE shouldn't do Present.
+}
+#endif // #if defined(OVR_GL)
+#endif // OVR_DIRECT_RENDERING
+
+
+#endif // OCULUS_RIFT_SUPPORTED_PLATFORMS
\ No newline at end of file
diff --git a/Engine/Plugins/Runtime/OculusRift/Source/OculusRift/Public/IOculusRiftPlugin.h b/Engine/Plugins/Runtime/OculusRift/Source/OculusRift/Public/IOculusRiftPlugin.h
new file mode 100644
--- /dev/null
+++ b/Engine/Plugins/Runtime/OculusRift/Source/OculusRift/Public/IOculusRiftPlugin.h
@@ -0,0 +1,41 @@
+// Copyright 1998-2013 Epic Games, Inc. All Rights Reserved.
+
+#pragma once
+
+#include "ModuleManager.h"
+#include "IHeadMountedDisplayModule.h"
+
+// Oculus support is not available on windows xp
+#define OCULUS_RIFT_SUPPORTED_PLATFORMS (PLATFORM_WINDOWS && WINVER > 0x0502) || PLATFORM_MAC
+
+/**
+ * The public interface to this module.  In most cases, this interface is only public to sibling modules 
+ * within this plugin.
+ */
+class IOculusRiftPlugin : public IHeadMountedDisplayModule
+{
+
+public:
+
+	/**
+	 * Singleton-like access to this module's interface.  This is just for convenience!
+	 * Beware of calling this during the shutdown phase, though.  Your module might have been unloaded already.
+	 *
+	 * @return Returns singleton instance, loading the module on demand if needed
+	 */
+	static inline IOculusRiftPlugin& Get()
+	{
+		return FModuleManager::LoadModuleChecked< IOculusRiftPlugin >( "OculusRift" );
+	}
+
+	/**
+	 * Checks to see if this module is loaded and ready.  It is only valid to call Get() if IsAvailable() returns true.
+	 *
+	 * @return True if the module is loaded and ready to use
+	 */
+	static inline bool IsAvailable()
+	{
+		return FModuleManager::Get().IsModuleLoaded( "OculusRift" );
+	}
+};
+
