#oculus Integrating fix for crash on render thread

diff --git a/Engine/Plugins/Runtime/OculusRift/Source/OculusRift/Private/OculusRiftHMD.cpp b/Engine/Plugins/Runtime/OculusRift/Source/OculusRift/Private/OculusRiftHMD.cpp
--- a/Engine/Plugins/Runtime/OculusRift/Source/OculusRift/Private/OculusRiftHMD.cpp
+++ b/Engine/Plugins/Runtime/OculusRift/Source/OculusRift/Private/OculusRiftHMD.cpp
@@ -1026,7 +1026,7 @@ bool FOculusRiftHMD::Exec( UWorld* InWorld, const TCHAR* Cmd, FOutputDevice& Ar
 	}
 	else if (FParse::Command(&Cmd, TEXT("OVRVERSION")))
 	{
-		static const char* Results = ovr_GetVersionString() + 7; // +7 is temporary, until libovr returns just numbers
+		static const char* Results = ovr_GetVersionString();
 		Ar.Logf(TEXT("%s, LibOVR: %s, built %s, %s"), *GEngineVersion.ToString(), UTF8_TO_TCHAR(Results), 
 			UTF8_TO_TCHAR(__DATE__), UTF8_TO_TCHAR(__TIME__));
 		return true;
@@ -1099,8 +1099,8 @@ void FOculusRiftHMD::RecordAnalytics()
 		EventAttributes.Add(FAnalyticsEventAttribute(TEXT("AllowFinishCurrentFrame"), bAllowFinishCurrentFrame));
 #ifdef OVR_VISION_ENABLED
 		EventAttributes.Add(FAnalyticsEventAttribute(TEXT("HmdPosTracking"), bHmdPosTracking));
+#endif
 		EventAttributes.Add(FAnalyticsEventAttribute(TEXT("LowPersistenceMode"), bLowPersistenceMode));
-#endif		
 		EventAttributes.Add(FAnalyticsEventAttribute(TEXT("UpdateOnRT"), bUpdateOnRT));
 		EventAttributes.Add(FAnalyticsEventAttribute(TEXT("Overdrive"), bOverdrive));
 		EventAttributes.Add(FAnalyticsEventAttribute(TEXT("MirrorToWindow"), bMirrorToWindow));
@@ -1612,6 +1612,8 @@ void FOculusRiftHMD::Startup()
 			MotionPredictionInSeconds *= 2;
 		}
 
+		bHmdPosTracking = (SupportedTrackingCaps & ovrTrackingCap_Position) != 0;
+
 		UpdateDistortionCaps();
 		UpdateHmdRenderInfo();
 		UpdateStereoRenderingParams();
@@ -1935,12 +1937,6 @@ void FOculusRiftHMD::LoadFromIni()
 	{
 		bAllowFinishCurrentFrame = v;
 	}
-#ifdef OVR_VISION_ENABLED
-	if (GConfig->GetBool(OculusSettings, TEXT("bHmdPosTracking"), v, GEngineIni))
-	{
-		bHmdPosTracking = v;
-	}
-#endif // #ifdef OVR_VISION_ENABLED
 	if (GConfig->GetBool(OculusSettings, TEXT("bLowPersistenceMode"), v, GEngineIni))
 	{
 		bLowPersistenceMode = v;
@@ -1992,9 +1988,6 @@ void FOculusRiftHMD::SaveToIni()
 	}
 	GConfig->SetBool(OculusSettings, TEXT("bAllowFinishCurrentFrame"), bAllowFinishCurrentFrame, GEngineIni);
 
-#ifdef OVR_VISION_ENABLED
-	GConfig->SetBool(OculusSettings, TEXT("bHmdPosTracking"), bHmdPosTracking, GEngineIni);
-#endif
 	GConfig->SetBool(OculusSettings, TEXT("bLowPersistenceMode"), bLowPersistenceMode, GEngineIni);
 
 	GConfig->SetBool(OculusSettings, TEXT("bUpdateOnRT"), bUpdateOnRT, GEngineIni);
diff --git a/Engine/Plugins/Runtime/OculusRift/Source/OculusRift/Private/OculusRiftRender.cpp b/Engine/Plugins/Runtime/OculusRift/Source/OculusRift/Private/OculusRiftRender.cpp
--- a/Engine/Plugins/Runtime/OculusRift/Source/OculusRift/Private/OculusRiftRender.cpp
+++ b/Engine/Plugins/Runtime/OculusRift/Source/OculusRift/Private/OculusRiftRender.cpp
@@ -339,6 +339,17 @@ void FOculusRiftHMD::DrawDebug(UCanvas* Canvas, EStereoscopicPass StereoPass)
 	{
 		if (bDrawGrid)
 		{
+			bool bPopTransform = false;
+			if (EyeRenderDesc[0].DistortedViewport.Size.w != FMath::CeilToInt(Canvas->ClipX / 2) ||
+				EyeRenderDesc[0].DistortedViewport.Size.h != Canvas->ClipY)
+			{
+				bPopTransform = true;
+				Canvas->Canvas->PushAbsoluteTransform(FScaleMatrix(
+					FVector((Canvas->ClipX * 0.5f) / float(EyeRenderDesc[0].DistortedViewport.Size.w),
+					Canvas->ClipY / float(EyeRenderDesc[0].DistortedViewport.Size.h),
+					1.0f)));
+			}
+
 			const FColor cNormal(255, 0, 0);
 			const FColor cSpacer(255, 255, 0);
 			const FColor cMid(0, 128, 255);
@@ -362,7 +373,7 @@ void FOculusRiftHMD::DrawDebug(UCanvas* Canvas, EStereoscopicPass StereoPass)
 				limitX = Alg::Max(renderViewportW - midX, midX);
 				limitY = Alg::Max(renderViewportH - midY, midY);
 
-				int spacerMask = (lineStep << 2) - 1;
+				int spacerMask = (lineStep << 1) - 1;
 
 				for (int xp = 0; xp < limitX; xp += lineStep)
 				{
@@ -415,6 +426,10 @@ void FOculusRiftHMD::DrawDebug(UCanvas* Canvas, EStereoscopicPass StereoPass)
 					}
 				}
 			}
+			if (bPopTransform)
+			{
+				Canvas->Canvas->PopTransform();
+			}
 		}
 		return;
 	}
@@ -580,6 +595,7 @@ void FOculusRiftHMD::UpdateViewport(bool bUseSeparateRenderTarget, const FViewpo
 #ifdef OVR_DIRECT_RENDERING
 void FOculusRiftHMD::ShutdownRendering()
 {
+	check(IsInRenderingThread());
 #if defined(OVR_D3D_VERSION) && (OVR_D3D_VERSION == 11)
 	if (pD3D11Bridge)
 	{
@@ -612,6 +628,25 @@ FOculusRiftHMD::D3D11Bridge::D3D11Bridge(FOculusRiftHMD* plugin):
 void FOculusRiftHMD::D3D11Bridge::BeginRendering()
 {
 	check(IsInRenderingThread());
+
+	ID3D11Device* D3DDevice = (ID3D11Device*)RHIGetNativeDevice();
+	ID3D11DeviceContext* D3DDeviceContext = nullptr;
+	if (D3DDevice)
+	{
+		D3DDevice->GetImmediateContext(&D3DDeviceContext);
+	}
+	if (!bInitialized || D3DDevice != Cfg.D3D11.pDevice || D3DDeviceContext != Cfg.D3D11.pDeviceContext)
+	{
+		OVR::Lock::Locker lock(&ModifyLock);
+		Cfg.D3D11.Header.API = ovrRenderAPI_D3D11;
+		Cfg.D3D11.Header.Multisample = 1; //?? RenderParams.Multisample;
+		// Note, neither Device nor Context are AddRef-ed here. Not sure, if we need to.
+		Cfg.D3D11.pDevice = D3DDevice;
+		Cfg.D3D11.pDeviceContext = D3DDeviceContext;
+		bNeedReinitRendererAPI = true;
+		bInitialized = true;
+	}
+
 	if (bInitialized)
 	{
 		if (bNeedReinitRendererAPI)
@@ -679,17 +714,18 @@ void FOculusRiftHMD::D3D11Bridge::BeginRendering()
 void FOculusRiftHMD::D3D11Bridge::FinishRendering()
 {
 	check(IsInRenderingThread());
-	if (Plugin->RenderParams_RenderThread.bFrameBegun)
+
+	if (Plugin->RenderParams_RenderThread.bFrameBegun && !bNeedReinitEyeTextures) 
 	{
 		// Finish the frame and let OVR do buffer swap (Present) and flush/sync.
 		const ovrTexture eyeTextures[2] = { EyeTexture_RenderThread[0].Texture, EyeTexture_RenderThread[1].Texture };
 		ovrHmd_EndFrame(Plugin->Hmd, Plugin->RenderParams_RenderThread.EyeRenderPose, eyeTextures); // This function will present
-		Plugin->RenderParams_RenderThread.bFrameBegun = false;
 	}
 	else
 	{
 		UE_LOG(LogHMD, Warning, TEXT("Skipping frame: FinishRendering called with no corresponding BeginRendering (was BackBuffer re-allocated?)"));
 	}
+	Plugin->RenderParams_RenderThread.bFrameBegun = false;
 }
 
 void FOculusRiftHMD::D3D11Bridge::Reset_RenderThread()
@@ -780,25 +816,6 @@ void FOculusRiftHMD::D3D11Bridge::UpdateViewport(const FViewport& Viewport, FRHI
 	const uint32 RTSizeX = RT->GetSizeX();
 	const uint32 RTSizeY = RT->GetSizeY();
 
-	ID3D11Device* D3DDevice = (ID3D11Device*)RHIGetNativeDevice();
-	ID3D11DeviceContext* D3DDeviceContext = nullptr;
-	if (D3DDevice)
-	{
-		D3DDevice->GetImmediateContext(&D3DDeviceContext);
-	}
-	if (D3DDevice != Cfg.D3D11.pDevice || D3DDeviceContext != Cfg.D3D11.pDeviceContext)
-	{
-		OVR::Lock::Locker lock(&ModifyLock);
-		Cfg.D3D11.Header.API = ovrRenderAPI_D3D11;
-		Cfg.D3D11.Header.Multisample = 1; //?? RenderParams.Multisample;
-		// Note, neither Device nor Context are AddRef-ed here. Not sure, if we need to.
-		Cfg.D3D11.pDevice = D3DDevice;
-		Cfg.D3D11.pDeviceContext = D3DDeviceContext;
-		bNeedReinitRendererAPI = true;
-		Plugin->RenderParams_RenderThread.bFrameBegun = false;
-		bInitialized = true;
-	}
-
 	if (Cfg.D3D11.pBackBufferRT != pD3DBBRT ||
 		Cfg.D3D11.pSwapChain != pD3DSC ||
 		Cfg.D3D11.Header.RTSize.w != Viewport.GetSizeXY().X ||
@@ -938,6 +955,17 @@ void FOculusRiftHMD::OGLBridge::FinishRendering()
 {
 	check(IsInRenderingThread());
 
+	if (bNeedReinitEyeTextures)
+	{
+		// make sure we use most recent textures, otherwise there will 
+		// be an assertion.
+		OVR::Lock::Locker lock(&ModifyEyeTexturesLock);
+
+		EyeTexture_RenderThread[0] = EyeTexture[0];
+		EyeTexture_RenderThread[1] = EyeTexture[1];
+		bNeedReinitEyeTextures = false;
+	}
+
 	if (Plugin->RenderParams_RenderThread.bFrameBegun)
 	{
 		// Finish the frame and let OVR do buffer swap (Present) and flush/sync.
@@ -961,6 +989,8 @@ void FOculusRiftHMD::OGLBridge::Init()
 
 void FOculusRiftHMD::OGLBridge::Reset()
 {
+	check(IsInRenderingThread());
+
 	EyeTexture[0].OGL.TexId = 0;
 	EyeTexture[1].OGL.TexId = 0;
 	EyeTexture_RenderThread[0].OGL.TexId = 0;
