#oculus Fix for improper shutdown of the direct rendering path, which caused a crash

diff --git a/Engine/Plugins/Runtime/OculusRift/Source/OculusRift/Private/OculusRiftHMD.cpp b/Engine/Plugins/Runtime/OculusRift/Source/OculusRift/Private/OculusRiftHMD.cpp
--- a/Engine/Plugins/Runtime/OculusRift/Source/OculusRift/Private/OculusRiftHMD.cpp
+++ b/Engine/Plugins/Runtime/OculusRift/Source/OculusRift/Private/OculusRiftHMD.cpp
@@ -1762,7 +1762,7 @@ void FOculusRiftHMD::Shutdown()
 
 	RestoreSystemValues();
 
-#ifndef OVR_SDK_RENDERING
+#ifdef OVR_SDK_RENDERING
 	ENQUEUE_UNIQUE_RENDER_COMMAND_ONEPARAMETER(ShutdownRen,
 	FOculusRiftHMD*, Plugin, this,
 	{
