#oculus Integrating fix for unsupported screenpercentage on ES2 feature sets

diff --git a/Engine/Plugins/Runtime/OculusRift/Source/OculusRift/Private/OculusRiftHMD.cpp b/Engine/Plugins/Runtime/OculusRift/Source/OculusRift/Private/OculusRiftHMD.cpp
--- a/Engine/Plugins/Runtime/OculusRift/Source/OculusRift/Private/OculusRiftHMD.cpp
+++ b/Engine/Plugins/Runtime/OculusRift/Source/OculusRift/Private/OculusRiftHMD.cpp
@@ -1397,20 +1397,23 @@ void FOculusRiftHMD::RestoreSystemValues()
 
 void FOculusRiftHMD::UpdateScreenSettings(const FViewport*)
 {
-	// Set the current ScreenPercentage state
-	static IConsoleVariable* CScrPercVar = IConsoleManager::Get().FindConsoleVariable(TEXT("r.ScreenPercentage"));
-	float DesiredSceeenPercentage;
-	if (Flags.bOverrideScreenPercentage)
-	{
-		DesiredSceeenPercentage = ScreenPercentage;
-	}
-	else
+	if (Flags.bScreenPercentageEnabled)
 	{
-		DesiredSceeenPercentage = IdealScreenPercentage;
-	}
-	if (FMath::RoundToInt(CScrPercVar->GetFloat()) != FMath::RoundToInt(DesiredSceeenPercentage))
-	{
-		CScrPercVar->Set(DesiredSceeenPercentage);
+		// Set the current ScreenPercentage state
+		static IConsoleVariable* CScrPercVar = IConsoleManager::Get().FindConsoleVariable(TEXT("r.ScreenPercentage"));
+		float DesiredSceeenPercentage;
+		if (Flags.bOverrideScreenPercentage)
+		{
+			DesiredSceeenPercentage = ScreenPercentage;
+		}
+		else
+		{
+			DesiredSceeenPercentage = IdealScreenPercentage;
+		}
+		if (FMath::RoundToInt(CScrPercVar->GetFloat()) != FMath::RoundToInt(DesiredSceeenPercentage))
+		{
+			CScrPercVar->Set(DesiredSceeenPercentage);
+		}
 	}
 }
 
@@ -1675,7 +1678,6 @@ void FOculusRiftHMD::ModifyShowFlags(FEngineShowFlags& ShowFlags)
 #else
 	ShowFlags.HMDDistortion = false;
 #endif
-	ShowFlags.ScreenPercentage = true;
 	ShowFlags.StereoRendering = IsStereoEnabled();
 }
 
@@ -1687,6 +1689,7 @@ void FOculusRiftHMD::SetupView(FSceneViewFamily& InViewFamily, FSceneView& InVie
 	{
 		WorldToMetersScale = InView.WorldToMetersScale;
 	}
+	Flags.bScreenPercentageEnabled = InViewFamily.EngineShowFlags.ScreenPercentage;
 
 #ifndef OVR_SDK_RENDERING
 	InViewFamily.bUseSeparateRenderTarget = false;
diff --git a/Engine/Plugins/Runtime/OculusRift/Source/OculusRift/Private/OculusRiftHMD.h b/Engine/Plugins/Runtime/OculusRift/Source/OculusRift/Private/OculusRiftHMD.h
--- a/Engine/Plugins/Runtime/OculusRift/Source/OculusRift/Private/OculusRiftHMD.h
+++ b/Engine/Plugins/Runtime/OculusRift/Source/OculusRift/Private/OculusRiftHMD.h
@@ -509,6 +509,9 @@ private: // data
 			/** Whether timewarp is enabled or not */
 			bool64 bTimeWarp : 1;
 
+			/** Whether ScreenPercentage usage is enabled */
+			bool64 bScreenPercentageEnabled : 1;
+
 			/** True, if pos tracking is enabled */
 			bool64				bHmdPosTracking : 1;
 
diff --git a/Engine/Plugins/Runtime/OculusRift/Source/OculusRift/Private/OculusRiftRender.cpp b/Engine/Plugins/Runtime/OculusRift/Source/OculusRift/Private/OculusRiftRender.cpp
--- a/Engine/Plugins/Runtime/OculusRift/Source/OculusRift/Private/OculusRiftRender.cpp
+++ b/Engine/Plugins/Runtime/OculusRift/Source/OculusRift/Private/OculusRiftRender.cpp
@@ -267,12 +267,15 @@ void FOculusRiftHMD::CalculateRenderTargetSize(uint32& InOutSizeX, uint32& InOut
 {
 	check(IsInGameThread());
 
-	static const auto CVar = IConsoleManager::Get().FindTConsoleVariableDataFloat(TEXT("r.ScreenPercentage"));
-	float value = CVar->GetValueOnGameThread();
-	if (value > 0.0f)
+	if (Flags.bScreenPercentageEnabled)
 	{
-		InOutSizeX = FMath::CeilToInt(InOutSizeX * value/100.f);
-		InOutSizeY = FMath::CeilToInt(InOutSizeY * value/100.f);
+		static const auto CVar = IConsoleManager::Get().FindTConsoleVariableDataFloat(TEXT("r.ScreenPercentage"));
+		float value = CVar->GetValueOnGameThread();
+		if (value > 0.0f)
+		{
+			InOutSizeX = FMath::CeilToInt(InOutSizeX * value / 100.f);
+			InOutSizeY = FMath::CeilToInt(InOutSizeY * value / 100.f);
+		}
 	}
 }
 
diff --git a/Engine/Source/Runtime/Engine/Private/SceneView.cpp b/Engine/Source/Runtime/Engine/Private/SceneView.cpp
--- a/Engine/Source/Runtime/Engine/Private/SceneView.cpp
+++ b/Engine/Source/Runtime/Engine/Private/SceneView.cpp
@@ -1091,13 +1091,6 @@ void FSceneView::EndFinalPostprocessSettings(const FSceneViewInitOptions& ViewIn
 		{
 			FinalPostProcessSettings.ScreenPercentage = Value;
 		}
-
-		// Not supported in ES2.
-		auto FeatureLevel = Family->Scene->GetFeatureLevel();
-		if(FeatureLevel == ERHIFeatureLevel::ES2 || FeatureLevel == ERHIFeatureLevel::ES3_1)
-		{
-			FinalPostProcessSettings.ScreenPercentage = 100.0f;
-		}
 	}
 
 #if !(UE_BUILD_SHIPPING || UE_BUILD_TEST)
@@ -1429,6 +1422,13 @@ FSceneViewFamily::FSceneViewFamily( const ConstructionValues& CVS )
 	LandscapeLODOverride = -1;
 	bDrawBaseInfo = true;
 #endif
+
+	// Not supported in ES2.
+	auto FeatureLevel = GetFeatureLevel();
+	if (FeatureLevel == ERHIFeatureLevel::ES2 || FeatureLevel == ERHIFeatureLevel::ES3_1)
+	{
+		EngineShowFlags.ScreenPercentage = false;
+	}
 }
 
 void FSceneViewFamily::ComputeFamilySize()
