#oculus Eliminating 'no oculus detected' warning

diff --git a/Engine/Plugins/Runtime/OculusRift/Source/OculusRift/Private/OculusRiftHMD.cpp b/Engine/Plugins/Runtime/OculusRift/Source/OculusRift/Private/OculusRiftHMD.cpp
--- a/Engine/Plugins/Runtime/OculusRift/Source/OculusRift/Private/OculusRiftHMD.cpp
+++ b/Engine/Plugins/Runtime/OculusRift/Source/OculusRift/Private/OculusRiftHMD.cpp
@@ -1938,10 +1938,7 @@ bool FOculusRiftHMD::InitDevice()
 		UpdateStereoRenderingParams();
 		UpdateHmdCaps();
 	}
-	else
-	{
-		UE_LOG(LogHMD, Warning, TEXT("No Oculus HMD detected! Is Oculus Run-Time installed and service is running?"));
-	}
+
 	return Hmd != nullptr;
 }
 
