Integrated 2386648 from Oculus branch.

- Fixes direct mode blanking out when losing and then regaining focus.

diff --git a/Engine/Plugins/Runtime/OculusRift/Source/OculusRift/Private/OculusRiftHMD.cpp b/Engine/Plugins/Runtime/OculusRift/Source/OculusRift/Private/OculusRiftHMD.cpp
--- a/Engine/Plugins/Runtime/OculusRift/Source/OculusRift/Private/OculusRiftHMD.cpp
+++ b/Engine/Plugins/Runtime/OculusRift/Source/OculusRift/Private/OculusRiftHMD.cpp
@@ -1263,40 +1263,50 @@ bool FOculusRiftHMD::DoEnableStereo(bool bStereo, bool bApplyToHmd)
 				}
 			}
 
-			if (bApplyToHmd && IsFullscreenAllowed() && SceneVP)
+			if (SceneVP)
 			{
 				TSharedPtr<SWindow> Window = SceneVP->FindWindow();
 				if (Window.IsValid())
 				{
 					FVector2D size = Window->GetSizeInScreen();
-					SceneVP->SetViewportSize(size.X, size.Y);
-					Window->SetViewportSizeDrivenByWindow(true);
 
-					if (stereoEnabled)
-					{
-						EWindowMode::Type wm = (!GIsEditor) ? EWindowMode::Fullscreen : EWindowMode::WindowedFullscreen;
-						FVector2D size = Window->GetSizeInScreen();
-						SceneVP->ResizeFrame(size.X, size.Y, wm, 0, 0);
-					}
-					else
+					if (bApplyToHmd && IsFullscreenAllowed())
 					{
-						// In Editor we cannot use ResizeFrame trick since it is called too late and App::IsGame
-						// returns false.
-						if (GIsEditor)
+						SceneVP->SetViewportSize(size.X, size.Y);
+						Window->SetViewportSizeDrivenByWindow(true);
+
+						if (stereoEnabled)
 						{
-							FSlateRect PreFullScreenRect;
-							PopPreFullScreenRect(PreFullScreenRect);
-							if (PreFullScreenRect.GetSize().X > 0 && PreFullScreenRect.GetSize().Y > 0 && IsFullscreenAllowed())
-							{
-								Window->MoveWindowTo(FVector2D(PreFullScreenRect.Left, PreFullScreenRect.Top));
-							}
+							EWindowMode::Type wm = (!GIsEditor) ? EWindowMode::Fullscreen : EWindowMode::WindowedFullscreen;
+							FVector2D size = Window->GetSizeInScreen();
+							SceneVP->ResizeFrame(size.X, size.Y, wm, 0, 0);
 						}
 						else
 						{
-							FVector2D size = Window->GetSizeInScreen();
-							SceneVP->ResizeFrame(size.X, size.Y, EWindowMode::Windowed, 0, 0);
+							// In Editor we cannot use ResizeFrame trick since it is called too late and App::IsGame
+							// returns false.
+							if (GIsEditor)
+							{
+								FSlateRect PreFullScreenRect;
+								PopPreFullScreenRect(PreFullScreenRect);
+								if (PreFullScreenRect.GetSize().X > 0 && PreFullScreenRect.GetSize().Y > 0 && IsFullscreenAllowed())
+								{
+									Window->MoveWindowTo(FVector2D(PreFullScreenRect.Left, PreFullScreenRect.Top));
+								}
+							}
+							else
+							{
+								FVector2D size = Window->GetSizeInScreen();
+								SceneVP->ResizeFrame(size.X, size.Y, EWindowMode::Windowed, 0, 0);
+							}
 						}
 					}
+					else if (!IsFullscreenAllowed())
+					{
+						// a special case when 'stereo on' or 'stereo hmd' is used in Direct mode. We must set proper window mode, otherwise
+						// it will be lost once window loses and regains the focus.
+						FSystemResolution::RequestResolutionChange(size.X, size.Y, (stereoEnabled) ? EWindowMode::WindowedMirror : EWindowMode::Windowed);
+					}
 				}
 			}
 		}
