Fixed bad usabge of cvars in Stereo rendering code, other cvars->Set() should go as well - if possible.

diff --git a/Engine/Plugins/Experimental/SimpleHMD/Source/SimpleHMD/Private/SimpleHMD.cpp b/Engine/Plugins/Experimental/SimpleHMD/Source/SimpleHMD/Private/SimpleHMD.cpp
--- a/Engine/Plugins/Experimental/SimpleHMD/Source/SimpleHMD/Private/SimpleHMD.cpp
+++ b/Engine/Plugins/Experimental/SimpleHMD/Source/SimpleHMD/Private/SimpleHMD.cpp
@@ -271,10 +271,6 @@ void FSimpleHMD::DrawDistortionMesh_RenderThread(struct FRenderingCompositePassC
 	DrawIndexedPrimitiveUP(Context.RHICmdList, PT_TriangleList, 0, NumVerts, NumTris, &Indices,
 		sizeof(Indices[0]), &Verts, sizeof(Verts[0]));
 }
-	
-void FSimpleHMD::UpdateScreenSettings(const FViewport*)
-{
-}
 
 bool FSimpleHMD::IsStereoEnabled() const
 {
diff --git a/Engine/Plugins/Experimental/SimpleHMD/Source/SimpleHMD/Private/SimpleHMD.h b/Engine/Plugins/Experimental/SimpleHMD/Source/SimpleHMD/Private/SimpleHMD.h
--- a/Engine/Plugins/Experimental/SimpleHMD/Source/SimpleHMD/Private/SimpleHMD.h
+++ b/Engine/Plugins/Experimental/SimpleHMD/Source/SimpleHMD/Private/SimpleHMD.h
@@ -63,7 +63,6 @@ public:
 	virtual FVector GetPositionOffset() const override;
 
 	virtual void DrawDistortionMesh_RenderThread(struct FRenderingCompositePassContext& Context, const FSceneView& View, const FIntPoint& TextureSize) override;
-	virtual void UpdateScreenSettings(const FViewport*) override;
 
 	/** IStereoRendering interface */
 	virtual bool IsStereoEnabled() const override;
diff --git a/Engine/Plugins/Experimental/SteamVR/Source/SteamVR/Private/SteamVRHMD.cpp b/Engine/Plugins/Experimental/SteamVR/Source/SteamVR/Private/SteamVRHMD.cpp
--- a/Engine/Plugins/Experimental/SteamVR/Source/SteamVR/Private/SteamVRHMD.cpp
+++ b/Engine/Plugins/Experimental/SteamVR/Source/SteamVR/Private/SteamVRHMD.cpp
@@ -361,10 +361,6 @@ FVector FSteamVRHMD::GetPositionOffset() const
 	return FVector::ZeroVector;
 }
 
-void FSteamVRHMD::UpdateScreenSettings(const FViewport*)
-{
-}
-
 bool FSteamVRHMD::IsStereoEnabled() const
 {
 	return bStereoEnabled && bHmdEnabled;
diff --git a/Engine/Plugins/Experimental/SteamVR/Source/SteamVR/Private/SteamVRHMD.h b/Engine/Plugins/Experimental/SteamVR/Source/SteamVR/Private/SteamVRHMD.h
--- a/Engine/Plugins/Experimental/SteamVR/Source/SteamVR/Private/SteamVRHMD.h
+++ b/Engine/Plugins/Experimental/SteamVR/Source/SteamVR/Private/SteamVRHMD.h
@@ -82,7 +82,6 @@ public:
 	virtual FVector GetPositionOffset() const override;
 
 	virtual void DrawDistortionMesh_RenderThread(struct FRenderingCompositePassContext& Context, const FSceneView& View, const FIntPoint& TextureSize) override;
-	virtual void UpdateScreenSettings(const FViewport*) override;
 
 	/** IStereoRendering interface */
 	virtual bool IsStereoEnabled() const override;
diff --git a/Engine/Plugins/Runtime/OculusRift/Source/OculusRift/Private/OculusRiftHMD.cpp b/Engine/Plugins/Runtime/OculusRift/Source/OculusRift/Private/OculusRiftHMD.cpp
--- a/Engine/Plugins/Runtime/OculusRift/Source/OculusRift/Private/OculusRiftHMD.cpp
+++ b/Engine/Plugins/Runtime/OculusRift/Source/OculusRift/Private/OculusRiftHMD.cpp
@@ -1388,42 +1388,24 @@ void FOculusRiftHMD::SaveSystemValues()
 {
 	static IConsoleVariable* CVSyncVar = IConsoleManager::Get().FindConsoleVariable(TEXT("r.VSync"));
 	Flags.bSavedVSync = CVSyncVar->GetInt() != 0;
-
-	static IConsoleVariable* CScrPercVar = IConsoleManager::Get().FindConsoleVariable(TEXT("r.ScreenPercentage"));
-	SavedScrPerc = CScrPercVar->GetFloat();
 }
 
 void FOculusRiftHMD::RestoreSystemValues()
 {
 	static IConsoleVariable* CVSyncVar = IConsoleManager::Get().FindConsoleVariable(TEXT("r.VSync"));
+	// todo: bad - cvars are a user wish, this should be changed
 	CVSyncVar->Set(Flags.bSavedVSync != 0);
 
-	static IConsoleVariable* CScrPercVar = IConsoleManager::Get().FindConsoleVariable(TEXT("r.ScreenPercentage"));
-	CScrPercVar->Set(SavedScrPerc);
-
 	static IConsoleVariable* CFinishFrameVar = IConsoleManager::Get().FindConsoleVariable(TEXT("r.FinishCurrentFrame"));
+	// todo: bad - cvars are a user wish, this should be changed
 	CFinishFrameVar->Set(false);
 }
 
-void FOculusRiftHMD::UpdateScreenSettings(const FViewport*)
+void FOculusRiftHMD::UpdatePostProcessSettings(FPostProcessSettings* Settings)
 {
 	if (Flags.bScreenPercentageEnabled)
 	{
-		// Set the current ScreenPercentage state
-		static IConsoleVariable* CScrPercVar = IConsoleManager::Get().FindConsoleVariable(TEXT("r.ScreenPercentage"));
-		float DesiredSceeenPercentage;
-		if (Flags.bOverrideScreenPercentage)
-		{
-			DesiredSceeenPercentage = ScreenPercentage;
-		}
-		else
-		{
-			DesiredSceeenPercentage = IdealScreenPercentage;
-		}
-		if (FMath::RoundToInt(CScrPercVar->GetFloat()) != FMath::RoundToInt(DesiredSceeenPercentage))
-		{
-			CScrPercVar->Set(DesiredSceeenPercentage);
-		}
+		Settings->ScreenPercentage = Flags.bOverrideScreenPercentage ? ScreenPercentage : IdealScreenPercentage;
 	}
 }
 
diff --git a/Engine/Plugins/Runtime/OculusRift/Source/OculusRift/Private/OculusRiftHMD.h b/Engine/Plugins/Runtime/OculusRift/Source/OculusRift/Private/OculusRiftHMD.h
--- a/Engine/Plugins/Runtime/OculusRift/Source/OculusRift/Private/OculusRiftHMD.h
+++ b/Engine/Plugins/Runtime/OculusRift/Source/OculusRift/Private/OculusRiftHMD.h
@@ -160,7 +160,7 @@ public:
 	virtual FVector GetPositionOffset() const override;
 
 	virtual void DrawDistortionMesh_RenderThread(struct FRenderingCompositePassContext& Context, const FSceneView& View, const FIntPoint& TextureSize) override;
-	virtual void UpdateScreenSettings(const FViewport*) override;
+	virtual void UpdatePostProcessSettings(FPostProcessSettings*) override;
 
 	virtual bool HandleInputKey(class UPlayerInput*, const FKey& Key, EInputEvent EventType, float AmountDepressed, bool bGamepad) override;
 
diff --git a/Engine/Source/Runtime/Engine/Private/GameViewportClient.cpp b/Engine/Source/Runtime/Engine/Private/GameViewportClient.cpp
--- a/Engine/Source/Runtime/Engine/Private/GameViewportClient.cpp
+++ b/Engine/Source/Runtime/Engine/Private/GameViewportClient.cpp
@@ -718,19 +718,17 @@ void UGameViewportClient::Draw(FViewport* InViewport, FCanvas* SceneCanvas)
 	DebugCanvasObject->Canvas = DebugCanvas;	
 	DebugCanvasObject->Init(InViewport->GetSizeXY().X, InViewport->GetSizeXY().Y, NULL);
 
-	const bool bScaledToRenderTarget = GEngine->HMDDevice.IsValid() && GEngine->IsStereoscopic3D(InViewport);
-	if (bScaledToRenderTarget)
 	{
-		// Allow HMD to modify screen settings
-		GEngine->HMDDevice->UpdateScreenSettings(Viewport);
-	}
-	if (DebugCanvas)
-	{
-		DebugCanvas->SetScaledToRenderTarget(bScaledToRenderTarget);
-	}
-	if (SceneCanvas)
-	{
-		SceneCanvas->SetScaledToRenderTarget(bScaledToRenderTarget);
+		const bool bScaledToRenderTarget = GEngine->HMDDevice.IsValid() && GEngine->IsStereoscopic3D(InViewport);
+
+		if (DebugCanvas)
+		{
+			DebugCanvas->SetScaledToRenderTarget(bScaledToRenderTarget);
+		}
+		if (SceneCanvas)
+		{
+			SceneCanvas->SetScaledToRenderTarget(bScaledToRenderTarget);
+		}
 	}
 
 	bool bUIDisableWorldRendering = false;
diff --git a/Engine/Source/Runtime/Engine/Private/SceneView.cpp b/Engine/Source/Runtime/Engine/Private/SceneView.cpp
--- a/Engine/Source/Runtime/Engine/Private/SceneView.cpp
+++ b/Engine/Source/Runtime/Engine/Private/SceneView.cpp
@@ -15,6 +15,7 @@
 #include "BufferVisualizationData.h"
 #include "Interfaces/Interface_PostProcessVolume.h"
 #include "Engine/TextureCube.h"
+#include "IHeadMountedDisplay.h"
 
 DEFINE_LOG_CATEGORY(LogBufferVisualization);
 
@@ -1100,6 +1101,15 @@ void FSceneView::EndFinalPostprocessSettings(const FSceneViewInitOptions& ViewIn
 		}
 	}
 
+	{
+		const bool bStereoEnabled = StereoPass != eSSP_FULL;
+		const bool bScaledToRenderTarget = GEngine->HMDDevice.IsValid() && bStereoEnabled;
+		if (bScaledToRenderTarget)
+		{
+			GEngine->HMDDevice->UpdatePostProcessSettings(&FinalPostProcessSettings);
+		}
+	}
+
 #if !(UE_BUILD_SHIPPING || UE_BUILD_TEST)
 	{
 		float Value = CVarSSRMaxRoughness.GetValueOnGameThread();
diff --git a/Engine/Source/Runtime/HeadMountedDisplay/Public/IHeadMountedDisplay.h b/Engine/Source/Runtime/HeadMountedDisplay/Public/IHeadMountedDisplay.h
--- a/Engine/Source/Runtime/HeadMountedDisplay/Public/IHeadMountedDisplay.h
+++ b/Engine/Source/Runtime/HeadMountedDisplay/Public/IHeadMountedDisplay.h
@@ -273,10 +273,16 @@ public:
 	virtual void DrawDistortionMesh_RenderThread(struct FRenderingCompositePassContext& Context, const FSceneView& View, const FIntPoint& TextureSize) {}
 
 	/**
-	 * This method is able to change screen settings (such as rendering scale) right before
-	 * any drawing occurs. It is called at the beginning of UGameViewportClient::Draw() method.
+	 * This method is able to change screen settings right before any drawing occurs. 
+	 * It is called at the beginning of UGameViewportClient::Draw() method.
+	 * We might remove this one as UpdatePostProcessSettings should be able to capture all needed cases
 	 */
-	virtual void UpdateScreenSettings(const FViewport* InViewport) = 0;
+	virtual void UpdateScreenSettings(const FViewport* InViewport) {}
+
+	/**
+	 * Allows to override the PostProcessSettings in the last moment e.g. allows up sampled 3D rendering
+	 */
+	virtual void UpdatePostProcessSettings(FPostProcessSettings*) {}
 
 	/**
 	 * Draw desired debug information related to the HMD system.
