#vr Fix for oculus crash on startup.  Initialization is now skipped if the device couldn't be started up, and bridges are torn down properly to avoid race condition when flushing rendering commands (UE-7249)

diff --git a/Engine/Plugins/Runtime/OculusRift/Source/OculusRift/Private/OculusRiftHMD.cpp b/Engine/Plugins/Runtime/OculusRift/Source/OculusRift/Private/OculusRiftHMD.cpp
--- a/Engine/Plugins/Runtime/OculusRift/Source/OculusRift/Private/OculusRiftHMD.cpp
+++ b/Engine/Plugins/Runtime/OculusRift/Source/OculusRift/Private/OculusRiftHMD.cpp
@@ -1799,7 +1799,7 @@ void FOculusRiftHMD::Startup()
 	// Initializes LibOVR. This LogMask_All enables maximum logging.
 	// Custom allocator can also be specified here.
 	// Actually, most likely, the ovr_Initialize is already called from PreInit.
-	ovr_Initialize();
+	int8 bWasInitialized = ovr_Initialize();
 
 #if !UE_BUILD_SHIPPING
 	// Should be changed to CAPI when available.
@@ -1807,6 +1807,26 @@ void FOculusRiftHMD::Startup()
 	OVR::Log::SetGlobalLog(&OcLog);
 #endif //#if !UE_BUILD_SHIPPING
 
+	if (GIsEditor)
+	{
+		Flags.bHeadTrackingEnforced = true;
+		//AlternateFrameRateDivider = 2;
+	}
+
+	bool bForced = FParse::Param(FCommandLine::Get(), TEXT("forcedrift"));
+	bool bInitialized = false;
+	if (!bForced)
+	{
+		bInitialized = InitDevice();
+	}
+
+	if (!bInitialized && !bForced)
+	{
+		ovr_Shutdown();
+		Flags.InitStatus = 0;
+		return;
+	}
+
 	// Uncap fps to enable FPS higher than 62
 	GEngine->bSmoothFrameRate = false;
 
@@ -1827,20 +1847,7 @@ void FOculusRiftHMD::Startup()
 #endif
 #endif // #ifdef OVR_SDK_RENDERING
 
-	if (GIsEditor)
-	{
-		Flags.bHeadTrackingEnforced = true;
-		//AlternateFrameRateDivider = 2;
-	}
-
-	bool forced = true;
-	if (!FParse::Param(FCommandLine::Get(), TEXT("forcedrift")))
-	{
-		InitDevice();
-		forced = false;
-	}
-
-	if (forced || Hmd)
+	if (bForced || Hmd)
 	{
 		Flags.InitStatus |= eInitialized;
 
@@ -1850,7 +1857,7 @@ void FOculusRiftHMD::Startup()
 
 void FOculusRiftHMD::Shutdown()
 {
-	if (!(Flags.InitStatus & eStartupExecuted))
+	if (!(Flags.InitStatus & eInitialized))
 	{
 		return;
 	}
@@ -1863,6 +1870,7 @@ void FOculusRiftHMD::Shutdown()
 	{
 		Plugin->ShutdownRendering();
 	});
+	FlushRenderingCommands();
 #endif // OVR_SDK_RENDERING
 	ReleaseDevice();
 	
