#oculus Integrating fix for screenpercentage in hmd

diff --git a/Engine/Plugins/Runtime/OculusRift/Source/OculusRift/Private/OculusRiftHMD.cpp b/Engine/Plugins/Runtime/OculusRift/Source/OculusRift/Private/OculusRiftHMD.cpp
--- a/Engine/Plugins/Runtime/OculusRift/Source/OculusRift/Private/OculusRiftHMD.cpp
+++ b/Engine/Plugins/Runtime/OculusRift/Source/OculusRift/Private/OculusRiftHMD.cpp
@@ -1405,7 +1405,7 @@ void FOculusRiftHMD::UpdatePostProcessSettings(FPostProcessSettings* Settings)
 {
 	if (Flags.bScreenPercentageEnabled)
 	{
-		Settings->ScreenPercentage = Flags.bOverrideScreenPercentage ? ScreenPercentage : IdealScreenPercentage;
+		Settings->ScreenPercentage = GetScreenPercentage();
 	}
 }
 
diff --git a/Engine/Plugins/Runtime/OculusRift/Source/OculusRift/Private/OculusRiftHMD.h b/Engine/Plugins/Runtime/OculusRift/Source/OculusRift/Private/OculusRiftHMD.h
--- a/Engine/Plugins/Runtime/OculusRift/Source/OculusRift/Private/OculusRiftHMD.h
+++ b/Engine/Plugins/Runtime/OculusRift/Source/OculusRift/Private/OculusRiftHMD.h
@@ -173,6 +173,8 @@ public:
 	void GetCurrentPose(FQuat& CurrentHmdOrientation, FVector& CurrentHmdPosition);
 	void BeginRendering_RenderThread();
 
+	float GetScreenPercentage() const { return Flags.bOverrideScreenPercentage ? ScreenPercentage : IdealScreenPercentage; }
+
 #ifdef OVR_SDK_RENDERING
 	class BridgeBaseImpl : public FRHICustomPresent
 	{
diff --git a/Engine/Plugins/Runtime/OculusRift/Source/OculusRift/Private/OculusRiftRender.cpp b/Engine/Plugins/Runtime/OculusRift/Source/OculusRift/Private/OculusRiftRender.cpp
--- a/Engine/Plugins/Runtime/OculusRift/Source/OculusRift/Private/OculusRiftRender.cpp
+++ b/Engine/Plugins/Runtime/OculusRift/Source/OculusRift/Private/OculusRiftRender.cpp
@@ -269,8 +269,7 @@ void FOculusRiftHMD::CalculateRenderTargetSize(uint32& InOutSizeX, uint32& InOut
 
 	if (Flags.bScreenPercentageEnabled)
 	{
-		static const auto CVar = IConsoleManager::Get().FindTConsoleVariableDataFloat(TEXT("r.ScreenPercentage"));
-		float value = CVar->GetValueOnGameThread();
+		float value = GetScreenPercentage();
 		if (value > 0.0f)
 		{
 			InOutSizeX = FMath::CeilToInt(InOutSizeX * value / 100.f);
@@ -486,8 +485,7 @@ void FOculusRiftHMD::DrawDebug(UCanvas* Canvas, EStereoscopicPass StereoPass)
 		Canvas->Canvas->DrawShadowedString(X, Y, *Str, Font, TextColor);
 
 		Y += RowHeight;
-		static IConsoleVariable* CScrPercVar = IConsoleManager::Get().FindConsoleVariable(TEXT("r.ScreenPercentage"));
-		int32 sp = (int32)CScrPercVar->GetFloat();
+		int32 sp = (int32)GetScreenPercentage();
 		Str = FString::Printf(TEXT("SP: %d"), sp);
 		Canvas->Canvas->DrawShadowedString(X, Y, *Str, Font, TextColor);
 
