GuardMalloc caught this abuse of ANSI_TO_TCHAR which cannot be used to assign to a variable.

diff --git a/Engine/Plugins/Runtime/OculusRift/Source/OculusRift/Private/OculusRiftHMD.cpp b/Engine/Plugins/Runtime/OculusRift/Source/OculusRift/Private/OculusRiftHMD.cpp
--- a/Engine/Plugins/Runtime/OculusRift/Source/OculusRift/Private/OculusRiftHMD.cpp
+++ b/Engine/Plugins/Runtime/OculusRift/Source/OculusRift/Private/OculusRiftHMD.cpp
@@ -87,8 +87,7 @@ public:
 		int32 len = FCStringAnsi::GetVarArgs(buf, sizeof(buf), sizeof(buf) / sizeof(ANSICHAR), fmt, argList);
 		if (len > 0 && buf[len - 1] == '\n') // truncate the trailing new-line char, since Logf adds its own
 			buf[len - 1] = '\0';
-		TCHAR* tbuf = ANSI_TO_TCHAR(buf);
-		GLog->Logf(TEXT("OCULUS: %s"), tbuf);
+		GLog->Logf(TEXT("OCULUS: %s"), ANSI_TO_TCHAR(buf));
 	}
 };
 
