- Fix for Oculus build break when using XP SDK.

diff --git a/Engine/Plugins/Runtime/OculusRift/Source/OculusRift/Private/OculusFunctionLibrary.cpp b/Engine/Plugins/Runtime/OculusRift/Source/OculusRift/Private/OculusFunctionLibrary.cpp
--- a/Engine/Plugins/Runtime/OculusRift/Source/OculusRift/Private/OculusFunctionLibrary.cpp
+++ b/Engine/Plugins/Runtime/OculusRift/Source/OculusRift/Private/OculusFunctionLibrary.cpp
@@ -1,5 +1,7 @@
 // Copyright 1998-2015 Epic Games, Inc. All Rights Reserved.
 
+#if OCULUS_RIFT_SUPPORTED_PLATFORMS
+
 #include "HMDPrivatePCH.h"
 #include "OculusFunctionLibrary.h"
 #include "OculusRiftHMD.h"
@@ -154,3 +156,5 @@ void UOculusFunctionLibrary::GetPlayerCameraManagerFollowHmd(bool& bFollowHmdOri
 		bFollowHmdPosition = OculusHMD->GetSettings()->Flags.bPlayerCameraManagerFollowsHmdPosition;
 	}
 }
+
+#endif // OCULUS_RIFT_SUPPORTED_PLATFORMS
\ No newline at end of file
diff --git a/Engine/Plugins/Runtime/OculusRift/Source/OculusRift/Private/OculusRiftCAPI.cpp b/Engine/Plugins/Runtime/OculusRift/Source/OculusRift/Private/OculusRiftCAPI.cpp
--- a/Engine/Plugins/Runtime/OculusRift/Source/OculusRift/Private/OculusRiftCAPI.cpp
+++ b/Engine/Plugins/Runtime/OculusRift/Source/OculusRift/Private/OculusRiftCAPI.cpp
@@ -1,6 +1,10 @@
 #include "HMDPrivatePCH.h"
 #include "OculusRiftHMD.h"
 
+#if OCULUS_RIFT_SUPPORTED_PLATFORMS
+
 #include <OVR_CAPIShim.c>
 #include <OVR_CAPI_Util.cpp>
 #include <OVR_StereoProjection.cpp>
+
+#endif // OCULUS_RIFT_SUPPORTED_PLATFORMS
diff --git a/Engine/Plugins/Runtime/OculusRift/Source/OculusRift/Public/OculusFunctionLibrary.h b/Engine/Plugins/Runtime/OculusRift/Source/OculusRift/Public/OculusFunctionLibrary.h
--- a/Engine/Plugins/Runtime/OculusRift/Source/OculusRift/Public/OculusFunctionLibrary.h
+++ b/Engine/Plugins/Runtime/OculusRift/Source/OculusRift/Public/OculusFunctionLibrary.h
@@ -1,6 +1,9 @@
 // Copyright 1998-2015 Epic Games, Inc. All Rights Reserved.
 
 #pragma once
+
+#if OCULUS_RIFT_SUPPORTED_PLATFORMS
+
 #include "OculusFunctionLibrary.generated.h"
 
 USTRUCT(BlueprintType, meta = (DisplayName = "HMD User Profile Data Field"))
@@ -156,3 +159,5 @@ class UOculusFunctionLibrary : public UBlueprintFunctionLibrary
 protected:
 	static class FOculusRiftHMD* GetOculusHMD();
 };
+
+#endif // OCULUS_RIFT_SUPPORTED_PLATFORMS
\ No newline at end of file
