Fix Windows XP build and Oculus.

diff --git a/Engine/Plugins/Runtime/OculusRift/Source/OculusRift/OculusRift.Build.cs b/Engine/Plugins/Runtime/OculusRift/Source/OculusRift/OculusRift.Build.cs
--- a/Engine/Plugins/Runtime/OculusRift/Source/OculusRift/OculusRift.Build.cs
+++ b/Engine/Plugins/Runtime/OculusRift/Source/OculusRift/OculusRift.Build.cs
@@ -37,7 +37,7 @@ namespace UnrealBuildTool.Rules
 			}
 
             // Currently, the Rift is only supported on windows and mac platforms
-            if (Target.Platform == UnrealTargetPlatform.Win32 || Target.Platform == UnrealTargetPlatform.Win64 || Target.Platform == UnrealTargetPlatform.Mac)
+            if ((Target.Platform == UnrealTargetPlatform.Win32 && !WindowsPlatform.IsWindowsXPSupported()) || Target.Platform == UnrealTargetPlatform.Win64 || Target.Platform == UnrealTargetPlatform.Mac)
             {
 				PrivateDependencyModuleNames.AddRange(new string[] { "LibOVR", "OpenGLDrv" });
 
diff --git a/Engine/Plugins/Runtime/OculusRift/Source/OculusRift/Private/OculusFunctionLibrary.cpp b/Engine/Plugins/Runtime/OculusRift/Source/OculusRift/Private/OculusFunctionLibrary.cpp
--- a/Engine/Plugins/Runtime/OculusRift/Source/OculusRift/Private/OculusFunctionLibrary.cpp
+++ b/Engine/Plugins/Runtime/OculusRift/Source/OculusRift/Private/OculusFunctionLibrary.cpp
@@ -11,17 +11,21 @@ UOculusFunctionLibrary::UOculusFunctionLibrary(const FObjectInitializer& ObjectI
 
 FOculusRiftHMD* UOculusFunctionLibrary::GetOculusHMD()
 {
+#if OCULUS_RIFT_SUPPORTED_PLATFORMS
 	if (GEngine->HMDDevice.IsValid() && (GEngine->HMDDevice->GetHMDDeviceType() == EHMDDeviceType::DT_OculusRift))
 	{
 		return static_cast<FOculusRiftHMD*>(GEngine->HMDDevice.Get());
 	}
 
 	UE_LOG(LogHMD, Warning, TEXT("Called Oculus function library on non-Oculus device or device was not found"));
+#endif // OCULUS_RIFT_SUPPORTED_PLATFORMS
+
 	return nullptr;
 }
 
 void UOculusFunctionLibrary::GetPose(FRotator& DeviceRotation, FVector& DevicePosition, FVector& NeckPosition, bool bUseOrienationForPlayerCamera, bool bUsePositionForPlayerCamera, const FVector PositionScale)
 {
+#if OCULUS_RIFT_SUPPORTED_PLATFORMS
 	FOculusRiftHMD* OculusHMD = GetOculusHMD();
 	if (OculusHMD && OculusHMD->IsHeadTrackingAllowed())
 	{
@@ -38,6 +42,7 @@ void UOculusFunctionLibrary::GetPose(FRotator& DeviceRotation, FVector& DevicePo
 		//UE_LOG(LogUHeadMountedDisplay, Log, TEXT("ROT: sYaw %.3f Pitch %.3f Roll %.3f"), DeviceRotation.Yaw, DeviceRotation.Pitch, DeviceRotation.Roll);
 	}
 	else
+#endif // OCULUS_RIFT_SUPPORTED_PLATFORMS
 	{
 		DeviceRotation = FRotator::ZeroRotator;
 		DevicePosition = FVector::ZeroVector;
@@ -46,6 +51,7 @@ void UOculusFunctionLibrary::GetPose(FRotator& DeviceRotation, FVector& DevicePo
 
 void UOculusFunctionLibrary::SetBaseRotationAndBaseOffsetInMeters(FRotator Rotation, FVector BaseOffsetInMeters, EOrientPositionSelector::Type Options)
 {
+#if OCULUS_RIFT_SUPPORTED_PLATFORMS
 	FOculusRiftHMD* OculusHMD = GetOculusHMD();
 	if ((OculusHMD != nullptr) && GEngine->HMDDevice->IsHeadTrackingAllowed())
 	{
@@ -58,10 +64,12 @@ void UOculusFunctionLibrary::SetBaseRotationAndBaseOffsetInMeters(FRotator Rotat
 			OculusHMD->SetBaseOffsetInMeters(BaseOffsetInMeters);
 		}
 	}
+#endif // OCULUS_RIFT_SUPPORTED_PLATFORMS
 }
 
 void UOculusFunctionLibrary::GetBaseRotationAndBaseOffsetInMeters(FRotator& OutRotation, FVector& OutBaseOffsetInMeters)
 {
+#if OCULUS_RIFT_SUPPORTED_PLATFORMS
 	FOculusRiftHMD* OculusHMD = GetOculusHMD();
 	if ((OculusHMD != nullptr) && GEngine->HMDDevice->IsHeadTrackingAllowed())
 	{
@@ -73,10 +81,12 @@ void UOculusFunctionLibrary::GetBaseRotationAndBaseOffsetInMeters(FRotator& OutR
 		OutRotation = FRotator::ZeroRotator;
 		OutBaseOffsetInMeters = FVector::ZeroVector;
 	}
+#endif // OCULUS_RIFT_SUPPORTED_PLATFORMS
 }
 
 void UOculusFunctionLibrary::GetRawSensorData(FVector& Accelerometer, FVector& Gyro, FVector& Magnetometer, float& Temperature, float& TimeInSeconds)
 {
+#if OCULUS_RIFT_SUPPORTED_PLATFORMS
 	FOculusRiftHMD* OculusHMD = GetOculusHMD();
 	if (OculusHMD != nullptr)
 	{
@@ -89,10 +99,12 @@ void UOculusFunctionLibrary::GetRawSensorData(FVector& Accelerometer, FVector& G
 		Temperature = Data.Temperature;
 		TimeInSeconds = Data.TimeInSeconds;
 	}
+#endif // OCULUS_RIFT_SUPPORTED_PLATFORMS
 }
 
 bool UOculusFunctionLibrary::GetUserProfile(FHmdUserProfile& Profile)
 {
+#if OCULUS_RIFT_SUPPORTED_PLATFORMS
 	FOculusRiftHMD* OculusHMD = GetOculusHMD();
 	if (OculusHMD != nullptr)
 	{
@@ -113,44 +125,53 @@ bool UOculusFunctionLibrary::GetUserProfile(FHmdUserProfile& Profile)
 			return true;
 		}
 	}
+#endif // OCULUS_RIFT_SUPPORTED_PLATFORMS
 	return false;
 }
 
 void UOculusFunctionLibrary::EnablePlayerControllerFollowHmd(bool bEnable)
 {
+#if OCULUS_RIFT_SUPPORTED_PLATFORMS
 	FOculusRiftHMD* OculusHMD = GetOculusHMD();
 	if (OculusHMD != nullptr)
 	{
 		OculusHMD->GetSettings()->Flags.bPlayerControllerFollowsHmd = bEnable;
 	}
+#endif // OCULUS_RIFT_SUPPORTED_PLATFORMS
 }
 
 bool UOculusFunctionLibrary::IsPlayerControllerFollowHmdEnabled()
 {
+#if OCULUS_RIFT_SUPPORTED_PLATFORMS
 	FOculusRiftHMD* OculusHMD = GetOculusHMD();
 	if (OculusHMD != nullptr)
 	{
 		return OculusHMD->GetSettings()->Flags.bPlayerControllerFollowsHmd != 0;
 	}
+#endif // OCULUS_RIFT_SUPPORTED_PLATFORMS
 	return true;
 }
 
 void UOculusFunctionLibrary::EnablePlayerCameraManagerFollowHmd(bool bFollowHmdOrientation, bool bFollowHmdPosition)
 {
+#if OCULUS_RIFT_SUPPORTED_PLATFORMS
 	FOculusRiftHMD* OculusHMD = GetOculusHMD();
 	if (OculusHMD != nullptr)
 	{
 		OculusHMD->GetSettings()->Flags.bPlayerCameraManagerFollowsHmdOrientation = bFollowHmdOrientation;
 		OculusHMD->GetSettings()->Flags.bPlayerCameraManagerFollowsHmdPosition = bFollowHmdPosition;
 	}
+#endif // OCULUS_RIFT_SUPPORTED_PLATFORMS
 }
 
 void UOculusFunctionLibrary::GetPlayerCameraManagerFollowHmd(bool& bFollowHmdOrientation, bool& bFollowHmdPosition)
 {
+#if OCULUS_RIFT_SUPPORTED_PLATFORMS
 	FOculusRiftHMD* OculusHMD = GetOculusHMD();
 	if (OculusHMD != nullptr)
 	{
 		bFollowHmdOrientation = OculusHMD->GetSettings()->Flags.bPlayerCameraManagerFollowsHmdOrientation;
 		bFollowHmdPosition = OculusHMD->GetSettings()->Flags.bPlayerCameraManagerFollowsHmdPosition;
 	}
+#endif // OCULUS_RIFT_SUPPORTED_PLATFORMS
 }
diff --git a/Engine/Plugins/Runtime/OculusRift/Source/OculusRift/Private/OculusRiftCAPI.cpp b/Engine/Plugins/Runtime/OculusRift/Source/OculusRift/Private/OculusRiftCAPI.cpp
--- a/Engine/Plugins/Runtime/OculusRift/Source/OculusRift/Private/OculusRiftCAPI.cpp
+++ b/Engine/Plugins/Runtime/OculusRift/Source/OculusRift/Private/OculusRiftCAPI.cpp
@@ -1,6 +1,8 @@
 #include "HMDPrivatePCH.h"
 #include "OculusRiftHMD.h"
 
+#if OCULUS_RIFT_SUPPORTED_PLATFORMS
 #include <OVR_CAPIShim.c>
 #include <OVR_CAPI_Util.cpp>
 #include <OVR_StereoProjection.cpp>
+#endif // OCULUS_RIFT_SUPPORTED_PLATFORMS
