#ue4 Fix for "error" log when Oculus HMD can't be found.  This is NOT a fatal error.

diff --git a/Engine/Plugins/Runtime/OculusRift/Source/OculusRift/Private/OculusRiftHMD.cpp b/Engine/Plugins/Runtime/OculusRift/Source/OculusRift/Private/OculusRiftHMD.cpp
--- a/Engine/Plugins/Runtime/OculusRift/Source/OculusRift/Private/OculusRiftHMD.cpp
+++ b/Engine/Plugins/Runtime/OculusRift/Source/OculusRift/Private/OculusRiftHMD.cpp
@@ -1031,7 +1031,7 @@ void FOculusRiftHMD::Startup()
 	if (!OVR_SUCCESS(initStatus) && initStatus == ovrError_LibLoad)
 	{
 		// fatal errors: can't load library
-		UE_LOG(LogHMD, Error, TEXT("Can't find Oculus library %s: is proper Runtime installed? Version: %s"), 
+		UE_LOG(LogHMD, Log, TEXT("Can't find Oculus library %s: is proper Runtime installed? Version: %s"), 
 			TEXT(OVR_FILE_DESCRIPTION_STRING), TEXT(OVR_VERSION_STRING));
 		Settings->Flags.InitStatus = 0;
 		return;
