#oculus Merging fix to allow for the 0.6.0.1 perf hud to work with oculus

diff --git a/Engine/Plugins/Runtime/OculusRift/Source/OculusRift/Private/OculusRiftHMD.cpp b/Engine/Plugins/Runtime/OculusRift/Source/OculusRift/Private/OculusRiftHMD.cpp
--- a/Engine/Plugins/Runtime/OculusRift/Source/OculusRift/Private/OculusRiftHMD.cpp
+++ b/Engine/Plugins/Runtime/OculusRift/Source/OculusRift/Private/OculusRiftHMD.cpp
@@ -471,6 +471,87 @@ bool FOculusRiftHMD::Exec( UWorld* InWorld, const TCHAR* Cmd, FOutputDevice& Ar
 			return true;
 		}
 #endif //UE_BUILD_SHIPPING
+		else
+		{
+			FString CmdName = FParse::Token(Cmd, 0);
+			if (CmdName.StartsWith(TEXT("SET")))
+			{
+				FString ValueNameStr = FParse::Token(Cmd, 0);
+				FString ValueStr = FParse::Token(Cmd, 0);
+
+				ovrBool res = ovrTrue;
+				CmdName = CmdName.Replace(TEXT("SET"), TEXT(""));
+				if (CmdName.Equals(TEXT("INT"), ESearchCase::IgnoreCase))
+				{
+					int v = FCString::Atoi(*ValueStr);
+					res = ovrHmd_SetInt(Hmd, TCHAR_TO_ANSI(*ValueNameStr), v);
+				}
+				else if (CmdName.Equals(TEXT("FLOAT"), ESearchCase::IgnoreCase))
+				{
+					float v = FCString::Atof(*ValueStr);
+					res = ovrHmd_SetFloat(Hmd, TCHAR_TO_ANSI(*ValueNameStr), v);
+				}
+				else if (CmdName.Equals(TEXT("BOOL"), ESearchCase::IgnoreCase))
+				{
+					ovrBool v;
+					if (ValueStr == TEXT("0") || ValueStr.Equals(TEXT("false"), ESearchCase::IgnoreCase))
+					{
+						v = ovrFalse;
+					}
+					else
+					{
+						v = ovrTrue;
+					}
+					res = ovrHmd_SetBool(Hmd, TCHAR_TO_ANSI(*ValueNameStr), v);
+				}
+				else if (CmdName.Equals(TEXT("STRING"), ESearchCase::IgnoreCase))
+				{
+					res = ovrHmd_SetString(Hmd, TCHAR_TO_ANSI(*ValueNameStr), TCHAR_TO_ANSI(*ValueStr));
+				}
+#if !UE_BUILD_SHIPPING
+				if (!res)
+				{
+					Ar.Logf(TEXT("HMD parameter %s was not set to value %s"), *ValueNameStr, *ValueStr);
+				}
+#endif // #if !UE_BUILD_SHIPPING
+				return true;
+			}
+#if !UE_BUILD_SHIPPING
+			else if (CmdName.StartsWith(TEXT("GET")))
+			{
+				FString ValueNameStr = FParse::Token(Cmd, 0);
+
+				CmdName = CmdName.Replace(TEXT("GET"), TEXT(""));
+				FString ValueStr;
+				if (CmdName.Equals(TEXT("INT"), ESearchCase::IgnoreCase))
+				{
+					int v = ovrHmd_GetInt(Hmd, TCHAR_TO_ANSI(*ValueNameStr), 0);
+					TCHAR buf[32];
+					FCString::Snprintf(buf, sizeof(buf) / sizeof(buf[0]), TEXT("%d"), v);
+					ValueStr = buf;
+				}
+				else if (CmdName.Equals(TEXT("FLOAT"), ESearchCase::IgnoreCase))
+				{
+					float v = ovrHmd_GetFloat(Hmd, TCHAR_TO_ANSI(*ValueNameStr), 0);
+					TCHAR buf[32];
+					FCString::Snprintf(buf, sizeof(buf)/sizeof(buf[0]), TEXT("%f"), v);
+					ValueStr = buf;
+				}
+				else if (CmdName.Equals(TEXT("BOOL"), ESearchCase::IgnoreCase))
+				{
+					ovrBool v = ovrHmd_GetBool(Hmd, TCHAR_TO_ANSI(*ValueNameStr), ovrFalse);
+					ValueStr = (v == ovrFalse) ? TEXT("false") : TEXT("true");
+				}
+				else if (CmdName.Equals(TEXT("STRING"), ESearchCase::IgnoreCase))
+				{
+					ValueStr = ANSI_TO_TCHAR(ovrHmd_GetString(Hmd, TCHAR_TO_ANSI(*ValueNameStr), ""));
+				}
+				Ar.Logf(TEXT("HMD parameter %s is set to value %s"), *ValueNameStr, *ValueStr);
+
+				return true;
+			}
+#endif // #if !UE_BUILD_SHIPPING
+		}
 	}
 	else if (FParse::Command(&Cmd, TEXT("HMDMAG")))
 	{
