UE-16177 Fix for crash when the oculus service is running, but no device detected.

diff --git a/Engine/Plugins/Runtime/OculusRift/Source/OculusRift/Private/OculusRiftHMD.cpp b/Engine/Plugins/Runtime/OculusRift/Source/OculusRift/Private/OculusRiftHMD.cpp
--- a/Engine/Plugins/Runtime/OculusRift/Source/OculusRift/Private/OculusRiftHMD.cpp
+++ b/Engine/Plugins/Runtime/OculusRift/Source/OculusRift/Private/OculusRiftHMD.cpp
@@ -1217,6 +1217,12 @@ bool FOculusRiftHMD::InitDevice()
 
 	FSettings* CurrentSettings = GetSettings();
 
+	ovrResult DetectResult = ovrHmd_Detect();
+	if (!OVR_SUCCESS(DetectResult))
+	{
+		return false;
+	}
+
 	ovrResult result = ovrHmd_Create(0, &Hmd);
 	if (OVR_SUCCESS(result) && Hmd)
 	{
