Fix for unsupported Oculus platforms

diff --git a/Engine/Plugins/Runtime/OculusRift/Source/OculusRift/Private/OculusFunctionLibrary.cpp b/Engine/Plugins/Runtime/OculusRift/Source/OculusRift/Private/OculusFunctionLibrary.cpp
--- a/Engine/Plugins/Runtime/OculusRift/Source/OculusRift/Private/OculusFunctionLibrary.cpp
+++ b/Engine/Plugins/Runtime/OculusRift/Source/OculusRift/Private/OculusFunctionLibrary.cpp
@@ -178,6 +178,7 @@ void UOculusFunctionLibrary::GetPlayerCameraManagerFollowHmd(bool& bFollowHmdOri
 
 void UOculusFunctionLibrary::SetBaseRotationAndPositionOffset(FRotator BaseRot, FVector PosOffset, EOrientPositionSelector::Type Options)
 {
+#if OCULUS_RIFT_SUPPORTED_PLATFORMS
 	FOculusRiftHMD* OculusHMD = GetOculusHMD();
 	if (OculusHMD != nullptr && GEngine->HMDDevice->IsHeadTrackingAllowed())
 	{
@@ -190,15 +191,18 @@ void UOculusFunctionLibrary::SetBaseRotationAndPositionOffset(FRotator BaseRot,
 			OculusHMD->GetSettings()->PositionOffset = PosOffset;
 		}
 	}
+#endif // OCULUS_RIFT_SUPPORTED_PLATFORMS
 }
 
 void UOculusFunctionLibrary::GetBaseRotationAndPositionOffset(FRotator& OutRot, FVector& OutPosOffset)
 {
+#if OCULUS_RIFT_SUPPORTED_PLATFORMS
 	FOculusRiftHMD* OculusHMD = GetOculusHMD();
 	if (OculusHMD != nullptr && GEngine->HMDDevice->IsHeadTrackingAllowed())
 	{
 		OutRot = OculusHMD->GetBaseRotation();
 		OutPosOffset = OculusHMD->GetSettings()->PositionOffset;
 	}
+#endif // OCULUS_RIFT_SUPPORTED_PLATFORMS
 }
 
