#oculus Integrating crash fix for when debug cam was used in VR preview

diff --git a/Engine/Plugins/Runtime/OculusRift/Source/OculusRift/Private/OculusRiftRender.cpp b/Engine/Plugins/Runtime/OculusRift/Source/OculusRift/Private/OculusRiftRender.cpp
--- a/Engine/Plugins/Runtime/OculusRift/Source/OculusRift/Private/OculusRiftRender.cpp
+++ b/Engine/Plugins/Runtime/OculusRift/Source/OculusRift/Private/OculusRiftRender.cpp
@@ -524,7 +524,7 @@ void FOculusRiftHMD::DrawDebug(UCanvas* Canvas)
 
 		//TODO:  Where can I get context!?
 		UWorld* MyWorld = GWorld;
-		if (FrameSettings->Flags.bDrawTrackingCameraFrustum)
+		if (Canvas && Canvas->SceneView && FrameSettings->Flags.bDrawTrackingCameraFrustum)
 		{
 			DrawDebugTrackingCameraFrustum(MyWorld, Canvas->SceneView->ViewRotation, Canvas->SceneView->ViewLocation);
 		}
