#UE4 Fixing Oculus cross eye effect when using IPD override.

diff --git a/Engine/Plugins/Runtime/OculusRift/Source/OculusRift/Private/OculusRiftHMD.cpp b/Engine/Plugins/Runtime/OculusRift/Source/OculusRift/Private/OculusRiftHMD.cpp
--- a/Engine/Plugins/Runtime/OculusRift/Source/OculusRift/Private/OculusRiftHMD.cpp
+++ b/Engine/Plugins/Runtime/OculusRift/Source/OculusRift/Private/OculusRiftHMD.cpp
@@ -1404,8 +1404,8 @@ void FOculusRiftHMD::UpdateStereoRenderingParams()
 		CurrentSettings->EyeRenderDesc[1] = ovrHmd_GetRenderDesc(Hmd, ovrEye_Right, CurrentSettings->EyeFov[1]);
 		if (CurrentSettings->Flags.bOverrideIPD)
 		{
-			CurrentSettings->EyeRenderDesc[0].HmdToEyeViewOffset.x = CurrentSettings->InterpupillaryDistance * 0.5f;
-			CurrentSettings->EyeRenderDesc[1].HmdToEyeViewOffset.x = -CurrentSettings->InterpupillaryDistance * 0.5f;
+			CurrentSettings->EyeRenderDesc[0].HmdToEyeViewOffset.x = -CurrentSettings->InterpupillaryDistance * 0.5f;
+			CurrentSettings->EyeRenderDesc[1].HmdToEyeViewOffset.x = CurrentSettings->InterpupillaryDistance * 0.5f;
 		}
 
 		const unsigned int ProjModifiers = ovrProjection_None; //@TODO revise to use ovrProjection_FarClipAtInfinity and/or ovrProjection_FarLessThanNear
