Maybe CIS Fix for the Mac

diff --git a/Engine/Plugins/Runtime/OculusRift/Source/OculusRift/Private/0.5/OculusRiftHMD_05.h b/Engine/Plugins/Runtime/OculusRift/Source/OculusRift/Private/0.5/OculusRiftHMD_05.h
--- a/Engine/Plugins/Runtime/OculusRift/Source/OculusRift/Private/0.5/OculusRiftHMD_05.h
+++ b/Engine/Plugins/Runtime/OculusRift/Source/OculusRift/Private/0.5/OculusRiftHMD_05.h
@@ -361,7 +361,7 @@ public:
 	*
 	* @param OutData	(out) SensorData structure to be filled in.
 	*/
-	virtual void GetRawSensorData(SensorData& OutData);
+	virtual void GetRawSensorData(SensorData& OutData) override;
 
 	/**
 	* User profile structure.
