Fix for build error

diff --git a/Engine/Plugins/Runtime/OculusRift/Source/OculusRift/Private/0.5/OculusRiftHMD_05.h b/Engine/Plugins/Runtime/OculusRift/Source/OculusRift/Private/0.5/OculusRiftHMD_05.h
--- a/Engine/Plugins/Runtime/OculusRift/Source/OculusRift/Private/0.5/OculusRiftHMD_05.h
+++ b/Engine/Plugins/Runtime/OculusRift/Source/OculusRift/Private/0.5/OculusRiftHMD_05.h
@@ -376,7 +376,7 @@ public:
 		FVector2D NeckToEyeDistance;	// Neck-to-eye distance, X - horizontal, Y - vertical, in meters
 		TMap<FString, FString> ExtraFields; // extra fields in name / value pairs.
 	};
-	virtual bool GetUserProfile(UserProfile& OutProfile);
+	virtual bool GetUserProfile(UserProfile& OutProfile) override;
 
 	virtual FString GetVersionString() const override;
 
