OSX Editor build fix

diff --git a/Engine/Plugins/Runtime/OculusRift/Source/OculusRift/Private/0.5/OculusRiftHMD_05.h b/Engine/Plugins/Runtime/OculusRift/Source/OculusRift/Private/0.5/OculusRiftHMD_05.h
--- a/Engine/Plugins/Runtime/OculusRift/Source/OculusRift/Private/0.5/OculusRiftHMD_05.h
+++ b/Engine/Plugins/Runtime/OculusRift/Source/OculusRift/Private/0.5/OculusRiftHMD_05.h
@@ -346,16 +346,6 @@ public:
 
 	virtual void DrawDebug(UCanvas* Canvas) override;
 
-	/* Raw sensor data structure. */
-	struct SensorData
-	{
-		FVector Accelerometer;	// Acceleration reading in m/s^2.
-		FVector Gyro;			// Rotation rate in rad/s.
-		FVector Magnetometer;   // Magnetic field in Gauss.
-		float Temperature;		// Temperature of the sensor in degrees Celsius.
-		float TimeInSeconds;	// Time when the reported IMU reading took place, in seconds.
-	};
-
 	/**
 	* Reports raw sensor data. If HMD doesn't support any of the parameters then it should be set to zero.
 	*
@@ -366,16 +356,6 @@ public:
 	/**
 	* User profile structure.
 	*/
-	struct UserProfile
-	{
-		FString Name;
-		FString Gender;
-		float PlayerHeight;				// Height of the player, in meters
-		float EyeHeight;				// Height of the player's eyes, in meters
-		float IPD;						// Interpupillary distance, in meters
-		FVector2D NeckToEyeDistance;	// Neck-to-eye distance, X - horizontal, Y - vertical, in meters
-		TMap<FString, FString> ExtraFields; // extra fields in name / value pairs.
-	};
 	virtual bool GetUserProfile(UserProfile& OutProfile) override;
 
 	virtual FString GetVersionString() const override;
