#oculus Fix for performance being poor on Oculus devices (r.finishcurrentframe erroneously enabled), fixed crash when device was disconnected

diff --git a/Engine/Plugins/Runtime/OculusRift/Source/OculusRift/Private/OculusRiftHMD.cpp b/Engine/Plugins/Runtime/OculusRift/Source/OculusRift/Private/OculusRiftHMD.cpp
--- a/Engine/Plugins/Runtime/OculusRift/Source/OculusRift/Private/OculusRiftHMD.cpp
+++ b/Engine/Plugins/Runtime/OculusRift/Source/OculusRift/Private/OculusRiftHMD.cpp
@@ -1698,6 +1698,16 @@ bool FOculusRiftHMD::GetUserProfile(UserProfile& OutProfile)
 	return false; 
 }
 
+void FOculusRiftHMD::ApplySystemOverridesOnStereo(bool force)
+{
+	check(IsInGameThread());
+	// ALWAYS SET r.FinishCurrentFrame to 0! Otherwise the perf might be poor.
+	// @TODO: revise the FD3D11DynamicRHI::RHIEndDrawingViewport code (and other renderers)
+	// to ignore this var completely.
+	static const auto CFinishFrameVar = IConsoleManager::Get().FindConsoleVariable(TEXT("r.FinishCurrentFrame"));
+	CFinishFrameVar->Set(0);
+}
+
 //////////////////////////////////////////////////////////////////////////
 FViewExtension::FViewExtension(FHeadMountedDisplay* InDelegate) 
 	: FHMDViewExtension(InDelegate)
diff --git a/Engine/Plugins/Runtime/OculusRift/Source/OculusRift/Private/OculusRiftHMD.h b/Engine/Plugins/Runtime/OculusRift/Source/OculusRift/Private/OculusRiftHMD.h
--- a/Engine/Plugins/Runtime/OculusRift/Source/OculusRift/Private/OculusRiftHMD.h
+++ b/Engine/Plugins/Runtime/OculusRift/Source/OculusRift/Private/OculusRiftHMD.h
@@ -440,6 +440,7 @@ private:
 	virtual void UpdateStereoRenderingParams() override;
 	virtual void UpdateHmdRenderInfo() override;
 	virtual void UpdateHmdCaps() override;
+	virtual void ApplySystemOverridesOnStereo(bool force = false) override;
 
 	/**
 	 * Called when state changes from 'stereo' to 'non-stereo'. Suppose to distribute
diff --git a/Engine/Plugins/Runtime/OculusRift/Source/OculusRift/Private/OculusRiftRenderD3D.cpp b/Engine/Plugins/Runtime/OculusRift/Source/OculusRift/Private/OculusRiftRenderD3D.cpp
--- a/Engine/Plugins/Runtime/OculusRift/Source/OculusRift/Private/OculusRiftRenderD3D.cpp
+++ b/Engine/Plugins/Runtime/OculusRift/Source/OculusRift/Private/OculusRiftRenderD3D.cpp
@@ -400,9 +400,13 @@ bool FOculusRiftHMD::D3D11Bridge::AllocateRenderTargetTexture(uint32 SizeX, uint
 			TexCreate_RenderTargetable | TexCreate_ShaderResource
 			);
 	}
-	OutTargetableTexture = ColorTextureSet->GetTexture2D();
-	OutShaderResourceTexture = ColorTextureSet->GetTexture2D();
-	return true;
+	if (ColorTextureSet)
+	{
+		OutTargetableTexture = ColorTextureSet->GetTexture2D();
+		OutShaderResourceTexture = ColorTextureSet->GetTexture2D();
+		return true;
+	}
+	return false;
 }
 
 void FOculusRiftHMD::D3D11Bridge::BeginRendering(FHMDViewExtension& InRenderContext, const FTexture2DRHIRef& RT)
