Added hidden and visible area meshes for the DK2 and Crescent Bay HMDs
Consolidated the visible and hidden area meshes into a single class that can be used by all HMD clients
Some house cleaning

diff --git a/Engine/Plugins/Runtime/OculusRift/Source/OculusRift/Private/OculusRiftHMD.cpp b/Engine/Plugins/Runtime/OculusRift/Source/OculusRift/Private/OculusRiftHMD.cpp
--- a/Engine/Plugins/Runtime/OculusRift/Source/OculusRift/Private/OculusRiftHMD.cpp
+++ b/Engine/Plugins/Runtime/OculusRift/Source/OculusRift/Private/OculusRiftHMD.cpp
@@ -2,6 +2,7 @@
 
 #include "HMDPrivatePCH.h"
 #include "OculusRiftHMD.h"
+#include "OculusRiftMeshAssets.h"
 
 #if !PLATFORM_MAC // Mac uses 0.5/OculusRiftHMD_05.cpp
 
@@ -1170,6 +1171,8 @@ void FOculusRiftHMD::Startup()
 		UE_LOG(LogHMD, Log, TEXT("Oculus plugin initialized. Version: %s"), *GetVersionString());
 	}
 
+	SetupOcclusionMeshes();
+
 	// grab a pointer to the renderer module for displaying our mirror window
 	static const FName RendererModuleName("Renderer");
 	RendererModule = FModuleManager::GetModulePtr<IRendererModule>(RendererModuleName);
@@ -1278,6 +1281,24 @@ void FOculusRiftHMD::ReleaseDevice()
 	}
 }
 
+void FOculusRiftHMD::SetupOcclusionMeshes()
+{
+	if (Hmd->Type == ovrHmdType::ovrHmd_DK2)
+	{
+		HiddenAreaMeshes[0].BuildMesh(DK2_LeftEyeHiddenAreaPositions, HiddenAreaVertexCount, FHMDViewMesh::MT_HiddenArea);
+		HiddenAreaMeshes[1].BuildMesh(DK2_RightEyeHiddenAreaPositions, HiddenAreaVertexCount, FHMDViewMesh::MT_HiddenArea);
+		VisibleAreaMeshes[0].BuildMesh(DK2_LeftEyeVisibleAreaPositions, VisibleAreaVertexCount, FHMDViewMesh::MT_VisibleArea);
+		VisibleAreaMeshes[1].BuildMesh(DK2_RightEyeVisibleAreaPositions, VisibleAreaVertexCount, FHMDViewMesh::MT_VisibleArea);
+	}
+	else if (Hmd->Type == ovrHmdType::ovrHmd_CB)
+	{
+		HiddenAreaMeshes[0].BuildMesh(CB_LeftEyeHiddenAreaPositions, HiddenAreaVertexCount, FHMDViewMesh::MT_HiddenArea);
+		HiddenAreaMeshes[1].BuildMesh(CB_RightEyeHiddenAreaPositions, HiddenAreaVertexCount, FHMDViewMesh::MT_HiddenArea);
+		VisibleAreaMeshes[0].BuildMesh(CB_LeftEyeVisibleAreaPositions, VisibleAreaVertexCount, FHMDViewMesh::MT_VisibleArea);
+		VisibleAreaMeshes[1].BuildMesh(CB_RightEyeVisibleAreaPositions, VisibleAreaVertexCount, FHMDViewMesh::MT_VisibleArea);
+	}
+}
+
 void FOculusRiftHMD::UpdateHmdCaps()
 {
 	if (Hmd)
diff --git a/Engine/Plugins/Runtime/OculusRift/Source/OculusRift/Private/OculusRiftHMD.h b/Engine/Plugins/Runtime/OculusRift/Source/OculusRift/Private/OculusRiftHMD.h
--- a/Engine/Plugins/Runtime/OculusRift/Source/OculusRift/Private/OculusRiftHMD.h
+++ b/Engine/Plugins/Runtime/OculusRift/Source/OculusRift/Private/OculusRiftHMD.h
@@ -286,6 +286,15 @@ public:
 	virtual bool IsFullscreenAllowed() override;
 	virtual void RecordAnalytics() override;
 
+	virtual bool HasHiddenAreaMesh() const override { return HiddenAreaMeshes[0].IsValid() && HiddenAreaMeshes[1].IsValid(); }
+	virtual void DrawHiddenAreaMesh_RenderThread(FRHICommandList& RHICmdList, EStereoscopicPass StereoPass) const override;
+
+	virtual bool HasVisibleAreaMesh() const override { return VisibleAreaMeshes[0].IsValid() && VisibleAreaMeshes[1].IsValid(); }
+	virtual void DrawVisibleAreaMesh_RenderThread(FRHICommandList& RHICmdList, EStereoscopicPass StereoPass) const override;
+
+	FHMDViewMesh HiddenAreaMeshes[2];
+	FHMDViewMesh VisibleAreaMeshes[2];
+
 	/** IStereoRendering interface */
 	virtual void RenderTexture_RenderThread(class FRHICommandListImmediate& RHICmdList, class FRHITexture2D* BackBuffer, class FRHITexture2D* SrcTexture) const override;
 	virtual void CalculateStereoViewOffset(const EStereoscopicPass StereoPassType, const FRotator& ViewRotation, 
@@ -434,6 +443,8 @@ private:
 
 	void ReleaseDevice();
 
+	void SetupOcclusionMeshes();
+
 	/**
 	 * Reads the device configuration, and sets up the stereoscopic rendering parameters
 	 */
diff --git a/Engine/Plugins/Runtime/OculusRift/Source/OculusRift/Private/OculusRiftMeshAssets.h b/Engine/Plugins/Runtime/OculusRift/Source/OculusRift/Private/OculusRiftMeshAssets.h
new file mode 100644
--- /dev/null
+++ b/Engine/Plugins/Runtime/OculusRift/Source/OculusRift/Private/OculusRiftMeshAssets.h
@@ -0,0 +1,830 @@
+// Copyright 1998-2015 Epic Gamesf, Inc. All Rights Reserved.
+
+#pragma once
+
+static const uint32 HiddenAreaVertexCount = 96;
+static const uint32 VisibleAreaVertexCount = 102;
+
+static const FVector2D DK2_LeftEyeHiddenAreaPositions[] =
+{
+	FVector2D(0.0f, 1.0f),
+	FVector2D(0.0f, 0.5671992302f),
+	FVector2D(0.01184238214f, 0.6673591137f),
+	FVector2D(0.3168017864f, 0.9884845018f),
+	FVector2D(0.4162660539f, 1.0f),
+	FVector2D(0.0f, 1.0f),
+	FVector2D(0.2371770293f, 0.9757383466f),
+	FVector2D(0.0f, 1.0f),
+	FVector2D(0.3168017864f, 0.9884845018f),
+	FVector2D(0.01184238214f, 0.6673591137f),
+	FVector2D(0.0f, 1.0f),
+	FVector2D(0.03479428217f, 0.7630788088f),
+	FVector2D(0.1613766551f, 0.9580639601f),
+	FVector2D(0.0f, 1.0f),
+	FVector2D(0.2371770293f, 0.9757383466f),
+	FVector2D(0.03479428217f, 0.7630788088f),
+	FVector2D(0.0f, 1.0f),
+	FVector2D(0.06616651267f, 0.8578283191f),
+	FVector2D(0.09944243729f, 0.9402768612f),
+	FVector2D(0.0f, 1.0f),
+	FVector2D(0.1613766551f, 0.9580639601f),
+	FVector2D(0.06616651267f, 0.8578283191f),
+	FVector2D(0.0f, 1.0f),
+	FVector2D(0.09944243729f, 0.9402768612f),
+	FVector2D(1.0f, 1.0f),
+	FVector2D(0.5424169302f, 1.0f),
+	FVector2D(0.6464179158f, 0.9915907979f),
+	FVector2D(0.6464179158f, 0.9915907979f),
+	FVector2D(1.0f, 1.0f),
+	FVector2D(0.740837872f, 0.9773967266f),
+	FVector2D(0.9785094857f, 0.7133126259f),
+	FVector2D(1.0f, 0.5973427892f),
+	FVector2D(1.0f, 1.0f),
+	FVector2D(0.740837872f, 0.9773967266f),
+	FVector2D(1.0f, 1.0f),
+	FVector2D(0.8373810649f, 0.9557206631f),
+	FVector2D(0.9517889023f, 0.8133317232f),
+	FVector2D(1.0f, 1.0f),
+	FVector2D(0.9785094857f, 0.7133126259f),
+	FVector2D(0.8373810649f, 0.9557206631f),
+	FVector2D(1.0f, 1.0f),
+	FVector2D(0.9018208385f, 0.9381420612f),
+	FVector2D(0.9265743494f, 0.8769161701f),
+	FVector2D(1.0f, 1.0f),
+	FVector2D(0.9517889023f, 0.8133317232f),
+	FVector2D(0.9018208385f, 0.9381420612f),
+	FVector2D(1.0f, 1.0f),
+	FVector2D(0.9265743494f, 0.8769161701f),
+	FVector2D(0.0f, 0.0f),
+	FVector2D(0.4123863876f, 0.0f),
+	FVector2D(0.2739078104f, 0.01678188145f),
+	FVector2D(0.2739078104f, 0.01678188145f),
+	FVector2D(0.0f, 0.0f),
+	FVector2D(0.2138241082f, 0.02765416726f),
+	FVector2D(0.0f, 0.4202009141f),
+	FVector2D(0.01011180785f, 0.3418060243f),
+	FVector2D(0.0f, 0.0f),
+	FVector2D(0.0284230262f, 0.2553999424f),
+	FVector2D(0.0f, 0.0f),
+	FVector2D(0.01011180785f, 0.3418060243f),
+	FVector2D(0.2138241082f, 0.02765416726f),
+	FVector2D(0.0f, 0.0f),
+	FVector2D(0.1548848748f, 0.04138758406f),
+	FVector2D(0.05703430623f, 0.1644160748f),
+	FVector2D(0.0f, 0.0f),
+	FVector2D(0.0284230262f, 0.2553999424f),
+	FVector2D(0.1548848748f, 0.04138758406f),
+	FVector2D(0.0f, 0.0f),
+	FVector2D(0.09995122254f, 0.05569322035f),
+	FVector2D(0.09995122254f, 0.05569322035f),
+	FVector2D(0.0f, 0.0f),
+	FVector2D(0.05703430623f, 0.1644160748f),
+	FVector2D(1.0f, 0.0f),
+	FVector2D(1.0f, 0.4133965373f),
+	FVector2D(0.9792567492f, 0.2858739495f),
+	FVector2D(0.5720955729f, 0.0f),
+	FVector2D(0.6576612592f, 0.008097825572f),
+	FVector2D(1.0f, 0.0f),
+	FVector2D(0.7528876066f, 0.02189874277f),
+	FVector2D(1.0f, 0.0f),
+	FVector2D(0.6576612592f, 0.008097825572f),
+	FVector2D(0.9792567492f, 0.2858739495f),
+	FVector2D(1.0f, 0.0f),
+	FVector2D(0.9574512243f, 0.2022407502f),
+	FVector2D(0.8432835937f, 0.04191007093f),
+	FVector2D(1.0f, 0.0f),
+	FVector2D(0.7528876066f, 0.02189874277f),
+	FVector2D(0.9574512243f, 0.2022407502f),
+	FVector2D(1.0f, 0.0f),
+	FVector2D(0.9358180165f, 0.1419302821f),
+	FVector2D(0.9029034376f, 0.05805703253f),
+	FVector2D(1.0f, 0.0f),
+	FVector2D(0.8432835937f, 0.04191007093f),
+	FVector2D(0.9358180165f, 0.1419302821f),
+	FVector2D(1.0f, 0.0f),
+	FVector2D(0.9029034376f, 0.05805703253f),
+};
+
+static const FVector2D DK2_RightEyeHiddenAreaPositions[] =
+{
+	FVector2D(1.0f, 1.0f),
+	FVector2D(0.5837339163f, 1.0f),
+	FVector2D(0.6831982136f, 0.9884845018f),
+	FVector2D(0.6831982136f, 0.9884845018f),
+	FVector2D(1.0f, 1.0f),
+	FVector2D(0.7628229856f, 0.9757383466f),
+	FVector2D(1.0f, 0.5671992302f),
+	FVector2D(0.98815763f, 0.6673591137f),
+	FVector2D(1.0f, 1.0f),
+	FVector2D(0.7628229856f, 0.9757383466f),
+	FVector2D(1.0f, 1.0f),
+	FVector2D(0.8386233449f, 0.9580639601f),
+	FVector2D(0.965205729f, 0.7630788088f),
+	FVector2D(1.0f, 1.0f),
+	FVector2D(0.98815763f, 0.6673591137f),
+	FVector2D(0.9338334799f, 0.8578283191f),
+	FVector2D(1.0f, 1.0f),
+	FVector2D(0.965205729f, 0.7630788088f),
+	FVector2D(0.8386233449f, 0.9580639601f),
+	FVector2D(1.0f, 1.0f),
+	FVector2D(0.9005575776f, 0.9402768612f),
+	FVector2D(0.9005575776f, 0.9402768612f),
+	FVector2D(1.0f, 1.0f),
+	FVector2D(0.9338334799f, 0.8578283191f),
+	FVector2D(0.0f, 1.0f),
+	FVector2D(0.0f, 0.5973427892f),
+	FVector2D(0.02149051428f, 0.7133126259f),
+	FVector2D(0.4575830698f, 1.0f),
+	FVector2D(0.3535820842f, 0.9915907979f),
+	FVector2D(0.0f, 1.0f),
+	FVector2D(0.259162128f, 0.9773967266f),
+	FVector2D(0.0f, 1.0f),
+	FVector2D(0.3535820842f, 0.9915907979f),
+	FVector2D(0.02149051428f, 0.7133126259f),
+	FVector2D(0.0f, 1.0f),
+	FVector2D(0.04821109772f, 0.8133317232f),
+	FVector2D(0.1626189351f, 0.9557206631f),
+	FVector2D(0.0f, 1.0f),
+	FVector2D(0.259162128f, 0.9773967266f),
+	FVector2D(0.09817916155f, 0.9381420612f),
+	FVector2D(0.0f, 1.0f),
+	FVector2D(0.1626189351f, 0.9557206631f),
+	FVector2D(0.04821109772f, 0.8133317232f),
+	FVector2D(0.0f, 1.0f),
+	FVector2D(0.0734256506f, 0.8769161701f),
+	FVector2D(0.0734256506f, 0.8769161701f),
+	FVector2D(0.0f, 1.0f),
+	FVector2D(0.09817916155f, 0.9381420612f),
+	FVector2D(1.0f, 0.0f),
+	FVector2D(1.0f, 0.4202009141f),
+	FVector2D(0.9898881912f, 0.3418060243f),
+	FVector2D(0.9898881912f, 0.3418060243f),
+	FVector2D(1.0f, 0.0f),
+	FVector2D(0.9715769887f, 0.2553999424f),
+	FVector2D(0.7260922194f, 0.01678188145f),
+	FVector2D(0.5876135826f, 0.0f),
+	FVector2D(1.0f, 0.0f),
+	FVector2D(0.7861759067f, 0.02765416726f),
+	FVector2D(1.0f, 0.0f),
+	FVector2D(0.7260922194f, 0.01678188145f),
+	FVector2D(0.9715769887f, 0.2553999424f),
+	FVector2D(1.0f, 0.0f),
+	FVector2D(0.9429656863f, 0.1644160748f),
+	FVector2D(0.8451151252f, 0.04138758406f),
+	FVector2D(1.0f, 0.0f),
+	FVector2D(0.7861759067f, 0.02765416726f),
+	FVector2D(0.9000487924f, 0.05569322035f),
+	FVector2D(1.0f, 0.0f),
+	FVector2D(0.8451151252f, 0.04138758406f),
+	FVector2D(0.9429656863f, 0.1644160748f),
+	FVector2D(1.0f, 0.0f),
+	FVector2D(0.9000487924f, 0.05569322035f),
+	FVector2D(0.0f, 0.0f),
+	FVector2D(0.4279044271f, 0.0f),
+	FVector2D(0.3423387408f, 0.008097825572f),
+	FVector2D(0.3423387408f, 0.008097825572f),
+	FVector2D(0.0f, 0.0f),
+	FVector2D(0.2471123934f, 0.02189874277f),
+	FVector2D(0.02074325085f, 0.2858739495f),
+	FVector2D(0.0f, 0.4133965373f),
+	FVector2D(0.0f, 0.0f),
+	FVector2D(0.2471123934f, 0.02189874277f),
+	FVector2D(0.0f, 0.0f),
+	FVector2D(0.1567164063f, 0.04191007093f),
+	FVector2D(0.04254877567f, 0.2022407502f),
+	FVector2D(0.0f, 0.0f),
+	FVector2D(0.02074325085f, 0.2858739495f),
+	FVector2D(0.1567164063f, 0.04191007093f),
+	FVector2D(0.0f, 0.0f),
+	FVector2D(0.09709656239f, 0.05805703253f),
+	FVector2D(0.06418198347f, 0.1419302821f),
+	FVector2D(0.0f, 0.0f),
+	FVector2D(0.04254877567f, 0.2022407502f),
+	FVector2D(0.09709656239f, 0.05805703253f),
+	FVector2D(0.0f, 0.0f),
+	FVector2D(0.06418198347f, 0.1419302821f),
+};
+
+static const FVector2D CB_LeftEyeHiddenAreaPositions[] =
+{
+	FVector2D(0.0f, 1.0f),
+	FVector2D(0.4649554491f, 1.0f),
+	FVector2D(0.3764877319f, 0.9927965403f),
+	FVector2D(0.3764877319f, 0.9927965403f),
+	FVector2D(0.0f, 1.0f),
+	FVector2D(0.2715978622f, 0.9779636264f),
+	FVector2D(0.0157302618f, 0.7411668301f),
+	FVector2D(0.0f, 0.6219738722f),
+	FVector2D(0.0f, 1.0f),
+	FVector2D(0.2715978622f, 0.9779636264f),
+	FVector2D(0.0f, 1.0f),
+	FVector2D(0.1714757681f, 0.9551845193f),
+	FVector2D(0.03109288216f, 0.8031471968f),
+	FVector2D(0.0f, 1.0f),
+	FVector2D(0.0157302618f, 0.7411668301f),
+	FVector2D(0.0496339798f, 0.8630086184f),
+	FVector2D(0.0f, 1.0f),
+	FVector2D(0.03109288216f, 0.8031471968f),
+	FVector2D(0.1714757681f, 0.9551845193f),
+	FVector2D(0.0f, 1.0f),
+	FVector2D(0.07188338041f, 0.919691503f),
+	FVector2D(0.07188338041f, 0.919691503f),
+	FVector2D(0.0f, 1.0f),
+	FVector2D(0.0496339798f, 0.8630086184f),
+	FVector2D(1.0f, 1.0f),
+	FVector2D(1.0f, 0.6241996288f),
+	FVector2D(0.9933852553f, 0.6894935966f),
+	FVector2D(0.9933852553f, 0.6894935966f),
+	FVector2D(1.0f, 1.0f),
+	FVector2D(0.9793937206f, 0.7837033868f),
+	FVector2D(0.7550624013f, 0.9875137806f),
+	FVector2D(0.6193442941f, 1.0f),
+	FVector2D(1.0f, 1.0f),
+	FVector2D(0.8259530067f, 0.9753877521f),
+	FVector2D(1.0f, 1.0f),
+	FVector2D(0.7550624013f, 0.9875137806f),
+	FVector2D(0.9793937206f, 0.7837033868f),
+	FVector2D(1.0f, 1.0f),
+	FVector2D(0.957940042f, 0.8690519333f),
+	FVector2D(0.8814527988f, 0.9627953768f),
+	FVector2D(1.0f, 1.0f),
+	FVector2D(0.8259530067f, 0.9753877521f),
+	FVector2D(0.9290241003f, 0.9497365952f),
+	FVector2D(1.0f, 1.0f),
+	FVector2D(0.8814527988f, 0.9627953768f),
+	FVector2D(0.957940042f, 0.8690519333f),
+	FVector2D(1.0f, 1.0f),
+	FVector2D(0.9290241003f, 0.9497365952f),
+	FVector2D(0.0f, 0.0f),
+	FVector2D(0.0f, 0.4600457847f),
+	FVector2D(0.01881575584f, 0.3334831297f),
+	FVector2D(0.3344950676f, 0.01416697446f),
+	FVector2D(0.4552387595f, 0.0f),
+	FVector2D(0.0f, 0.0f),
+	FVector2D(0.01881575584f, 0.3334831297f),
+	FVector2D(0.0f, 0.0f),
+	FVector2D(0.04645586014f, 0.2338332087f),
+	FVector2D(0.2493925691f, 0.03307863697f),
+	FVector2D(0.0f, 0.0f),
+	FVector2D(0.3344950676f, 0.01416697446f),
+	FVector2D(0.1802922487f, 0.05489978567f),
+	FVector2D(0.0f, 0.0f),
+	FVector2D(0.2493925691f, 0.03307863697f),
+	FVector2D(0.04645586014f, 0.2338332087f),
+	FVector2D(0.0f, 0.0f),
+	FVector2D(0.08427917957f, 0.1436391324f),
+	FVector2D(0.1228299141f, 0.07744830847f),
+	FVector2D(0.0f, 0.0f),
+	FVector2D(0.1802922487f, 0.05489978567f),
+	FVector2D(0.08427917957f, 0.1436391324f),
+	FVector2D(0.0f, 0.0f),
+	FVector2D(0.1228299141f, 0.07744830847f),
+	FVector2D(1.0f, 0.0f),
+	FVector2D(0.6072594523f, 0.0f),
+	FVector2D(0.6981808543f, 0.00689325761f),
+	FVector2D(0.6981808543f, 0.00689325761f),
+	FVector2D(1.0f, 0.0f),
+	FVector2D(0.7825559974f, 0.02071332f),
+	FVector2D(0.9825832248f, 0.3167535067f),
+	FVector2D(1.0f, 0.4462257326f),
+	FVector2D(1.0f, 0.0f),
+	FVector2D(0.7825559974f, 0.02071332f),
+	FVector2D(1.0f, 0.0f),
+	FVector2D(0.8487468362f, 0.03671549633f),
+	FVector2D(0.9571251869f, 0.216376245f),
+	FVector2D(1.0f, 0.0f),
+	FVector2D(0.9825832248f, 0.3167535067f),
+	FVector2D(0.8487468362f, 0.03671549633f),
+	FVector2D(1.0f, 0.0f),
+	FVector2D(0.882933259f, 0.0468987003f),
+	FVector2D(0.9265756011f, 0.131273821f),
+	FVector2D(1.0f, 0.0f),
+	FVector2D(0.9571251869f, 0.216376245f),
+	FVector2D(0.882933259f, 0.0468987003f),
+	FVector2D(1.0f, 0.0f),
+	FVector2D(0.9265756011f, 0.131273821f),
+};
+
+static const FVector2D CB_RightEyeHiddenAreaPositions[] =
+{
+	FVector2D(1.0f, 1.0f),
+	FVector2D(0.5350445509f, 1.0f),
+	FVector2D(0.6235122681f, 0.9927965403f),
+	FVector2D(0.6235122681f, 0.9927965403f),
+	FVector2D(1.0f, 1.0f),
+	FVector2D(0.7284021378f, 0.9779636264f),
+	FVector2D(0.9842697382f, 0.7411668301f),
+	FVector2D(1.0f, 0.6219738722f),
+	FVector2D(1.0f, 1.0f),
+	FVector2D(0.7284021378f, 0.9779636264f),
+	FVector2D(1.0f, 1.0f),
+	FVector2D(0.8285242319f, 0.9551845193f),
+	FVector2D(0.9689071178f, 0.8031471968f),
+	FVector2D(1.0f, 1.0f),
+	FVector2D(0.9842697382f, 0.7411668301f),
+	FVector2D(0.9503660202f, 0.8630086184f),
+	FVector2D(1.0f, 1.0f),
+	FVector2D(0.9689071178f, 0.8031471968f),
+	FVector2D(0.8285242319f, 0.9551845193f),
+	FVector2D(1.0f, 1.0f),
+	FVector2D(0.9281166196f, 0.919691503f),
+	FVector2D(0.9281166196f, 0.919691503f),
+	FVector2D(1.0f, 1.0f),
+	FVector2D(0.9503660202f, 0.8630086184f),
+	FVector2D(0.0f, 1.0f),
+	FVector2D(0.0f, 0.6241996288f),
+	FVector2D(0.006614744663f, 0.6894935966f),
+	FVector2D(0.006614744663f, 0.6894935966f),
+	FVector2D(0.0f, 1.0f),
+	FVector2D(0.02060627937f, 0.7837033868f),
+	FVector2D(0.2449375987f, 0.9875137806f),
+	FVector2D(0.3806557059f, 1.0f),
+	FVector2D(0.0f, 1.0f),
+	FVector2D(0.1740469933f, 0.9753877521f),
+	FVector2D(0.0f, 1.0f),
+	FVector2D(0.2449375987f, 0.9875137806f),
+	FVector2D(0.02060627937f, 0.7837033868f),
+	FVector2D(0.0f, 1.0f),
+	FVector2D(0.04205995798f, 0.8690519333f),
+	FVector2D(0.1185472012f, 0.9627953768f),
+	FVector2D(0.0f, 1.0f),
+	FVector2D(0.1740469933f, 0.9753877521f),
+	FVector2D(0.0709758997f, 0.9497365952f),
+	FVector2D(0.0f, 1.0f),
+	FVector2D(0.1185472012f, 0.9627953768f),
+	FVector2D(0.04205995798f, 0.8690519333f),
+	FVector2D(0.0f, 1.0f),
+	FVector2D(0.0709758997f, 0.9497365952f),
+	FVector2D(1.0f, 0.0f),
+	FVector2D(1.0f, 0.4600457847f),
+	FVector2D(0.9811842442f, 0.3334831297f),
+	FVector2D(0.6655049324f, 0.01416697446f),
+	FVector2D(0.5447612405f, 0.0f),
+	FVector2D(1.0f, 0.0f),
+	FVector2D(0.9811842442f, 0.3334831297f),
+	FVector2D(1.0f, 0.0f),
+	FVector2D(0.9535441399f, 0.2338332087f),
+	FVector2D(0.7506074309f, 0.03307863697f),
+	FVector2D(1.0f, 0.0f),
+	FVector2D(0.6655049324f, 0.01416697446f),
+	FVector2D(0.8197077513f, 0.05489978567f),
+	FVector2D(1.0f, 0.0f),
+	FVector2D(0.7506074309f, 0.03307863697f),
+	FVector2D(0.9535441399f, 0.2338332087f),
+	FVector2D(1.0f, 0.0f),
+	FVector2D(0.9157208204f, 0.1436391324f),
+	FVector2D(0.8771700859f, 0.07744830847f),
+	FVector2D(1.0f, 0.0f),
+	FVector2D(0.8197077513f, 0.05489978567f),
+	FVector2D(0.9157208204f, 0.1436391324f),
+	FVector2D(1.0f, 0.0f),
+	FVector2D(0.8771700859f, 0.07744830847f),
+	FVector2D(0.0f, 0.0f),
+	FVector2D(0.3927405477f, 0.0f),
+	FVector2D(0.3018191457f, 0.00689325761f),
+	FVector2D(0.3018191457f, 0.00689325761f),
+	FVector2D(0.0f, 0.0f),
+	FVector2D(0.2174440026f, 0.02071332f),
+	FVector2D(0.01741677523f, 0.3167535067f),
+	FVector2D(0.0f, 0.4462257326f),
+	FVector2D(0.0f, 0.0f),
+	FVector2D(0.2174440026f, 0.02071332f),
+	FVector2D(0.0f, 0.0f),
+	FVector2D(0.1512531638f, 0.03671549633f),
+	FVector2D(0.04287481308f, 0.216376245f),
+	FVector2D(0.0f, 0.0f),
+	FVector2D(0.01741677523f, 0.3167535067f),
+	FVector2D(0.1512531638f, 0.03671549633f),
+	FVector2D(0.0f, 0.0f),
+	FVector2D(0.117066741f, 0.0468987003f),
+	FVector2D(0.0734243989f, 0.131273821f),
+	FVector2D(0.0f, 0.0f),
+	FVector2D(0.04287481308f, 0.216376245f),
+	FVector2D(0.117066741f, 0.0468987003f),
+	FVector2D(0.0f, 0.0f),
+	FVector2D(0.0734243989f, 0.131273821f),
+};
+
+static const FVector2D DK2_LeftEyeVisibleAreaPositions[] =
+{
+	FVector2D(0.1548848748f, 0.04138758406f),
+	FVector2D(0.09995122254f, 0.05569322035f),
+	FVector2D(0.05703430623f, 0.1644160748f),
+	FVector2D(0.06616651267f, 0.8578283191f),
+	FVector2D(0.09944243729f, 0.9402768612f),
+	FVector2D(0.1613766551f, 0.9580639601f),
+	FVector2D(0.8373810649f, 0.9557206631f),
+	FVector2D(0.9018208385f, 0.9381420612f),
+	FVector2D(0.9265743494f, 0.8769161701f),
+	FVector2D(0.9358180165f, 0.1419302821f),
+	FVector2D(0.9029034376f, 0.05805703253f),
+	FVector2D(0.8432835937f, 0.04191007093f),
+	FVector2D(0.0f, 0.5671992302f),
+	FVector2D(0.0f, 0.4202009141f),
+	FVector2D(0.01184238214f, 0.6673591137f),
+	FVector2D(0.01184238214f, 0.6673591137f),
+	FVector2D(0.03479428217f, 0.7630788088f),
+	FVector2D(0.06616651267f, 0.8578283191f),
+	FVector2D(0.1613766551f, 0.9580639601f),
+	FVector2D(0.2371770293f, 0.9757383466f),
+	FVector2D(0.3168017864f, 0.9884845018f),
+	FVector2D(0.3168017864f, 0.9884845018f),
+	FVector2D(0.4162660539f, 1.0f),
+	FVector2D(0.5424169302f, 1.0f),
+	FVector2D(0.5424169302f, 1.0f),
+	FVector2D(0.6464179158f, 0.9915907979f),
+	FVector2D(0.740837872f, 0.9773967266f),
+	FVector2D(0.9265743494f, 0.8769161701f),
+	FVector2D(0.9517889023f, 0.8133317232f),
+	FVector2D(0.9785094857f, 0.7133126259f),
+	FVector2D(0.9785094857f, 0.7133126259f),
+	FVector2D(1.0f, 0.5973427892f),
+	FVector2D(1.0f, 0.4133965373f),
+	FVector2D(1.0f, 0.4133965373f),
+	FVector2D(0.9792567492f, 0.2858739495f),
+	FVector2D(0.9574512243f, 0.2022407502f),
+	FVector2D(0.8432835937f, 0.04191007093f),
+	FVector2D(0.7528876066f, 0.02189874277f),
+	FVector2D(0.6576612592f, 0.008097825572f),
+	FVector2D(0.6576612592f, 0.008097825572f),
+	FVector2D(0.5720955729f, 0.0f),
+	FVector2D(0.4123863876f, 0.0f),
+	FVector2D(0.4123863876f, 0.0f),
+	FVector2D(0.2739078104f, 0.01678188145f),
+	FVector2D(0.2138241082f, 0.02765416726f),
+	FVector2D(0.05703430623f, 0.1644160748f),
+	FVector2D(0.0284230262f, 0.2553999424f),
+	FVector2D(0.01011180785f, 0.3418060243f),
+	FVector2D(0.9358180165f, 0.1419302821f),
+	FVector2D(0.8432835937f, 0.04191007093f),
+	FVector2D(0.6576612592f, 0.008097825572f),
+	FVector2D(0.9574512243f, 0.2022407502f),
+	FVector2D(0.9358180165f, 0.1419302821f),
+	FVector2D(0.6576612592f, 0.008097825572f),
+	FVector2D(1.0f, 0.4133965373f),
+	FVector2D(0.9574512243f, 0.2022407502f),
+	FVector2D(0.6576612592f, 0.008097825572f),
+	FVector2D(1.0f, 0.4133965373f),
+	FVector2D(0.6576612592f, 0.008097825572f),
+	FVector2D(0.4123863876f, 0.0f),
+	FVector2D(0.9785094857f, 0.7133126259f),
+	FVector2D(1.0f, 0.4133965373f),
+	FVector2D(0.4123863876f, 0.0f),
+	FVector2D(0.9785094857f, 0.7133126259f),
+	FVector2D(0.4123863876f, 0.0f),
+	FVector2D(0.2138241082f, 0.02765416726f),
+	FVector2D(0.9265743494f, 0.8769161701f),
+	FVector2D(0.9785094857f, 0.7133126259f),
+	FVector2D(0.2138241082f, 0.02765416726f),
+	FVector2D(0.8373810649f, 0.9557206631f),
+	FVector2D(0.9265743494f, 0.8769161701f),
+	FVector2D(0.2138241082f, 0.02765416726f),
+	FVector2D(0.740837872f, 0.9773967266f),
+	FVector2D(0.8373810649f, 0.9557206631f),
+	FVector2D(0.2138241082f, 0.02765416726f),
+	FVector2D(0.5424169302f, 1.0f),
+	FVector2D(0.740837872f, 0.9773967266f),
+	FVector2D(0.2138241082f, 0.02765416726f),
+	FVector2D(0.3168017864f, 0.9884845018f),
+	FVector2D(0.5424169302f, 1.0f),
+	FVector2D(0.2138241082f, 0.02765416726f),
+	FVector2D(0.1613766551f, 0.9580639601f),
+	FVector2D(0.3168017864f, 0.9884845018f),
+	FVector2D(0.2138241082f, 0.02765416726f),
+	FVector2D(0.06616651267f, 0.8578283191f),
+	FVector2D(0.1613766551f, 0.9580639601f),
+	FVector2D(0.2138241082f, 0.02765416726f),
+	FVector2D(0.01184238214f, 0.6673591137f),
+	FVector2D(0.06616651267f, 0.8578283191f),
+	FVector2D(0.2138241082f, 0.02765416726f),
+	FVector2D(0.01184238214f, 0.6673591137f),
+	FVector2D(0.2138241082f, 0.02765416726f),
+	FVector2D(0.0f, 0.4202009141f),
+	FVector2D(0.01011180785f, 0.3418060243f),
+	FVector2D(0.2138241082f, 0.02765416726f),
+	FVector2D(0.0f, 0.4202009141f),
+	FVector2D(0.2138241082f, 0.02765416726f),
+	FVector2D(0.1548848748f, 0.04138758406f),
+	FVector2D(0.01011180785f, 0.3418060243f),
+	FVector2D(0.1548848748f, 0.04138758406f),
+	FVector2D(0.05703430623f, 0.1644160748f),
+	FVector2D(0.01011180785f, 0.3418060243f),
+};
+
+static const FVector2D DK2_RightEyeVisibleAreaPositions[] =
+{
+	FVector2D(0.9429656863f, 0.1644160748f),
+	FVector2D(0.9000487924f, 0.05569322035f),
+	FVector2D(0.8451151252f, 0.04138758406f),
+	FVector2D(0.0734256506f, 0.8769161701f),
+	FVector2D(0.09817916155f, 0.9381420612f),
+	FVector2D(0.1626189351f, 0.9557206631f),
+	FVector2D(0.8386233449f, 0.9580639601f),
+	FVector2D(0.9005575776f, 0.9402768612f),
+	FVector2D(0.9338334799f, 0.8578283191f),
+	FVector2D(0.1567164063f, 0.04191007093f),
+	FVector2D(0.09709656239f, 0.05805703253f),
+	FVector2D(0.06418198347f, 0.1419302821f),
+	FVector2D(0.0f, 0.5973427892f),
+	FVector2D(0.0f, 0.4133965373f),
+	FVector2D(0.02149051428f, 0.7133126259f),
+	FVector2D(0.02149051428f, 0.7133126259f),
+	FVector2D(0.04821109772f, 0.8133317232f),
+	FVector2D(0.0734256506f, 0.8769161701f),
+	FVector2D(0.1626189351f, 0.9557206631f),
+	FVector2D(0.259162128f, 0.9773967266f),
+	FVector2D(0.3535820842f, 0.9915907979f),
+	FVector2D(0.3535820842f, 0.9915907979f),
+	FVector2D(0.4575830698f, 1.0f),
+	FVector2D(0.5837339163f, 1.0f),
+	FVector2D(0.5837339163f, 1.0f),
+	FVector2D(0.6831982136f, 0.9884845018f),
+	FVector2D(0.7628229856f, 0.9757383466f),
+	FVector2D(0.9338334799f, 0.8578283191f),
+	FVector2D(0.965205729f, 0.7630788088f),
+	FVector2D(0.98815763f, 0.6673591137f),
+	FVector2D(0.98815763f, 0.6673591137f),
+	FVector2D(1.0f, 0.5671992302f),
+	FVector2D(1.0f, 0.4202009141f),
+	FVector2D(1.0f, 0.4202009141f),
+	FVector2D(0.9898881912f, 0.3418060243f),
+	FVector2D(0.9715769887f, 0.2553999424f),
+	FVector2D(0.8451151252f, 0.04138758406f),
+	FVector2D(0.7861759067f, 0.02765416726f),
+	FVector2D(0.7260922194f, 0.01678188145f),
+	FVector2D(0.7260922194f, 0.01678188145f),
+	FVector2D(0.5876135826f, 0.0f),
+	FVector2D(0.4279044271f, 0.0f),
+	FVector2D(0.4279044271f, 0.0f),
+	FVector2D(0.3423387408f, 0.008097825572f),
+	FVector2D(0.2471123934f, 0.02189874277f),
+	FVector2D(0.06418198347f, 0.1419302821f),
+	FVector2D(0.04254877567f, 0.2022407502f),
+	FVector2D(0.02074325085f, 0.2858739495f),
+	FVector2D(0.9429656863f, 0.1644160748f),
+	FVector2D(0.8451151252f, 0.04138758406f),
+	FVector2D(0.7260922194f, 0.01678188145f),
+	FVector2D(0.7628229856f, 0.9757383466f),
+	FVector2D(0.8386233449f, 0.9580639601f),
+	FVector2D(0.9338334799f, 0.8578283191f),
+	FVector2D(0.7628229856f, 0.9757383466f),
+	FVector2D(0.9338334799f, 0.8578283191f),
+	FVector2D(0.98815763f, 0.6673591137f),
+	FVector2D(0.5837339163f, 1.0f),
+	FVector2D(0.7628229856f, 0.9757383466f),
+	FVector2D(0.98815763f, 0.6673591137f),
+	FVector2D(0.5837339163f, 1.0f),
+	FVector2D(0.98815763f, 0.6673591137f),
+	FVector2D(1.0f, 0.4202009141f),
+	FVector2D(0.3535820842f, 0.9915907979f),
+	FVector2D(0.5837339163f, 1.0f),
+	FVector2D(1.0f, 0.4202009141f),
+	FVector2D(0.3535820842f, 0.9915907979f),
+	FVector2D(1.0f, 0.4202009141f),
+	FVector2D(0.9715769887f, 0.2553999424f),
+	FVector2D(0.1626189351f, 0.9557206631f),
+	FVector2D(0.3535820842f, 0.9915907979f),
+	FVector2D(0.9715769887f, 0.2553999424f),
+	FVector2D(0.0734256506f, 0.8769161701f),
+	FVector2D(0.1626189351f, 0.9557206631f),
+	FVector2D(0.9715769887f, 0.2553999424f),
+	FVector2D(0.02149051428f, 0.7133126259f),
+	FVector2D(0.0734256506f, 0.8769161701f),
+	FVector2D(0.9715769887f, 0.2553999424f),
+	FVector2D(0.02149051428f, 0.7133126259f),
+	FVector2D(0.9715769887f, 0.2553999424f),
+	FVector2D(0.0f, 0.4133965373f),
+	FVector2D(0.02074325085f, 0.2858739495f),
+	FVector2D(0.9715769887f, 0.2553999424f),
+	FVector2D(0.0f, 0.4133965373f),
+	FVector2D(0.9715769887f, 0.2553999424f),
+	FVector2D(0.06418198347f, 0.1419302821f),
+	FVector2D(0.02074325085f, 0.2858739495f),
+	FVector2D(0.9715769887f, 0.2553999424f),
+	FVector2D(0.1567164063f, 0.04191007093f),
+	FVector2D(0.06418198347f, 0.1419302821f),
+	FVector2D(0.2471123934f, 0.02189874277f),
+	FVector2D(0.9715769887f, 0.2553999424f),
+	FVector2D(0.1567164063f, 0.04191007093f),
+	FVector2D(0.9715769887f, 0.2553999424f),
+	FVector2D(0.4279044271f, 0.0f),
+	FVector2D(0.2471123934f, 0.02189874277f),
+	FVector2D(0.9715769887f, 0.2553999424f),
+	FVector2D(0.9429656863f, 0.1644160748f),
+	FVector2D(0.4279044271f, 0.0f),
+	FVector2D(0.9429656863f, 0.1644160748f),
+	FVector2D(0.7260922194f, 0.01678188145f),
+	FVector2D(0.4279044271f, 0.0f),
+};
+
+static const FVector2D CB_LeftEyeVisibleAreaPositions[] =
+{
+	FVector2D(0.9290241003f, 0.9497365952f),
+	FVector2D(0.957940042f, 0.8690519333f),
+	FVector2D(0.8814527988f, 0.9627953768f),
+	FVector2D(0.07188337296f, 0.919691503f),
+	FVector2D(0.1714757681f, 0.9551845193f),
+	FVector2D(0.04963401332f, 0.8630086184f),
+	FVector2D(0.882933259f, 0.0468987003f),
+	FVector2D(0.8487468362f, 0.03671549633f),
+	FVector2D(0.9265756011f, 0.131273821f),
+	FVector2D(0.1228299141f, 0.07744830847f),
+	FVector2D(0.08427921683f, 0.1436391324f),
+	FVector2D(0.1802922785f, 0.05489978567f),
+	FVector2D(0.4649554193f, 1.0f),
+	FVector2D(0.6193442941f, 1.0f),
+	FVector2D(0.3764877319f, 0.9927965403f),
+	FVector2D(0.271597892f, 0.9779636264f),
+	FVector2D(0.3764877319f, 0.9927965403f),
+	FVector2D(0.1714757681f, 0.9551845193f),
+	FVector2D(0.03109287843f, 0.8031471968f),
+	FVector2D(0.04963401332f, 0.8630086184f),
+	FVector2D(0.01573022455f, 0.7411668301f),
+	FVector2D(0.0f, 0.6219738722f),
+	FVector2D(0.01573022455f, 0.7411668301f),
+	FVector2D(0.0f, 0.4600457847f),
+	FVector2D(0.01881577075f, 0.3334831297f),
+	FVector2D(0.0f, 0.4600457847f),
+	FVector2D(0.04645589367f, 0.2338332087f),
+	FVector2D(0.249392584f, 0.03307863697f),
+	FVector2D(0.1802922785f, 0.05489978567f),
+	FVector2D(0.3344950676f, 0.01416697446f),
+	FVector2D(0.4552387595f, 0.0f),
+	FVector2D(0.3344950676f, 0.01416697446f),
+	FVector2D(0.6072594523f, 0.0f),
+	FVector2D(0.6981808543f, 0.00689325761f),
+	FVector2D(0.6072594523f, 0.0f),
+	FVector2D(0.7825559974f, 0.02071332f),
+	FVector2D(0.9571251869f, 0.216376245f),
+	FVector2D(0.9265756011f, 0.131273821f),
+	FVector2D(0.9825832248f, 0.3167535067f),
+	FVector2D(1.0f, 0.4462257326f),
+	FVector2D(0.9825832248f, 0.3167535067f),
+	FVector2D(1.0f, 0.6241996288f),
+	FVector2D(0.9933852553f, 0.6894935966f),
+	FVector2D(1.0f, 0.6241996288f),
+	FVector2D(0.9793937206f, 0.7837033868f),
+	FVector2D(0.8259530067f, 0.9753877521f),
+	FVector2D(0.8814527988f, 0.9627953768f),
+	FVector2D(0.7550624013f, 0.9875137806f),
+	FVector2D(0.7550624013f, 0.9875137806f),
+	FVector2D(0.8814527988f, 0.9627953768f),
+	FVector2D(0.957940042f, 0.8690519333f),
+	FVector2D(0.9793937206f, 0.7837033868f),
+	FVector2D(0.7550624013f, 0.9875137806f),
+	FVector2D(0.957940042f, 0.8690519333f),
+	FVector2D(0.7550624013f, 0.9875137806f),
+	FVector2D(0.9793937206f, 0.7837033868f),
+	FVector2D(1.0f, 0.6241996288f),
+	FVector2D(0.7550624013f, 0.9875137806f),
+	FVector2D(0.6193442941f, 1.0f),
+	FVector2D(1.0f, 0.6241996288f),
+	FVector2D(0.6193442941f, 1.0f),
+	FVector2D(1.0f, 0.6241996288f),
+	FVector2D(0.9825832248f, 0.3167535067f),
+	FVector2D(0.6193442941f, 1.0f),
+	FVector2D(0.9825832248f, 0.3167535067f),
+	FVector2D(0.3764877319f, 0.9927965403f),
+	FVector2D(0.9825832248f, 0.3167535067f),
+	FVector2D(0.1714757681f, 0.9551845193f),
+	FVector2D(0.3764877319f, 0.9927965403f),
+	FVector2D(0.9825832248f, 0.3167535067f),
+	FVector2D(0.04963401332f, 0.8630086184f),
+	FVector2D(0.1714757681f, 0.9551845193f),
+	FVector2D(0.9825832248f, 0.3167535067f),
+	FVector2D(0.01573022455f, 0.7411668301f),
+	FVector2D(0.04963401332f, 0.8630086184f),
+	FVector2D(0.9825832248f, 0.3167535067f),
+	FVector2D(0.0f, 0.4600457847f),
+	FVector2D(0.01573022455f, 0.7411668301f),
+	FVector2D(0.9825832248f, 0.3167535067f),
+	FVector2D(0.04645589367f, 0.2338332087f),
+	FVector2D(0.0f, 0.4600457847f),
+	FVector2D(0.9825832248f, 0.3167535067f),
+	FVector2D(0.9265756011f, 0.131273821f),
+	FVector2D(0.04645589367f, 0.2338332087f),
+	FVector2D(0.9265756011f, 0.131273821f),
+	FVector2D(0.8487468362f, 0.03671549633f),
+	FVector2D(0.04645589367f, 0.2338332087f),
+	FVector2D(0.7825559974f, 0.02071332f),
+	FVector2D(0.8487468362f, 0.03671549633f),
+	FVector2D(0.04645589367f, 0.2338332087f),
+	FVector2D(0.7825559974f, 0.02071332f),
+	FVector2D(0.6072594523f, 0.0f),
+	FVector2D(0.04645589367f, 0.2338332087f),
+	FVector2D(0.6072594523f, 0.0f),
+	FVector2D(0.3344950676f, 0.01416697446f),
+	FVector2D(0.04645589367f, 0.2338332087f),
+	FVector2D(0.04645589367f, 0.2338332087f),
+	FVector2D(0.3344950676f, 0.01416697446f),
+	FVector2D(0.08427921683f, 0.1436391324f),
+	FVector2D(0.3344950676f, 0.01416697446f),
+	FVector2D(0.1802922785f, 0.05489978567f),
+	FVector2D(0.08427921683f, 0.1436391324f),
+};
+
+static const FVector2D CB_RightEyeVisibleAreaPositions[] =
+{
+	FVector2D(0.1185472012f, 0.9627953768f),
+	FVector2D(0.0709758997f, 0.9497365952f),
+	FVector2D(0.04205995798f, 0.8690519333f),
+	FVector2D(0.9503660202f, 0.8630086184f),
+	FVector2D(0.9281166196f, 0.919691503f),
+	FVector2D(0.8285242319f, 0.9551845193f),
+	FVector2D(0.0734243989f, 0.131273821f),
+	FVector2D(0.117066741f, 0.0468987003f),
+	FVector2D(0.1512531638f, 0.03671549633f),
+	FVector2D(0.8197077513f, 0.05489978567f),
+	FVector2D(0.8771700859f, 0.07744830847f),
+	FVector2D(0.9157208204f, 0.1436391324f),
+	FVector2D(0.5350445509f, 1.0f),
+	FVector2D(0.6235122681f, 0.9927965403f),
+	FVector2D(0.3806557059f, 1.0f),
+	FVector2D(0.3806557059f, 1.0f),
+	FVector2D(0.2449375987f, 0.9875137806f),
+	FVector2D(0.1740469933f, 0.9753877521f),
+	FVector2D(0.04205995798f, 0.8690519333f),
+	FVector2D(0.02060627937f, 0.7837033868f),
+	FVector2D(0.006614744663f, 0.6894935966f),
+	FVector2D(0.006614744663f, 0.6894935966f),
+	FVector2D(0.0f, 0.6241996288f),
+	FVector2D(0.0f, 0.4462257326f),
+	FVector2D(0.0f, 0.4462257326f),
+	FVector2D(0.01741677523f, 0.3167535067f),
+	FVector2D(0.04287481308f, 0.216376245f),
+	FVector2D(0.1512531638f, 0.03671549633f),
+	FVector2D(0.2174440026f, 0.02071332f),
+	FVector2D(0.3018191457f, 0.00689325761f),
+	FVector2D(0.3018191457f, 0.00689325761f),
+	FVector2D(0.3927405477f, 0.0f),
+	FVector2D(0.5447612405f, 0.0f),
+	FVector2D(0.5447612405f, 0.0f),
+	FVector2D(0.6655049324f, 0.01416697446f),
+	FVector2D(0.7506074309f, 0.03307863697f),
+	FVector2D(0.9157208204f, 0.1436391324f),
+	FVector2D(0.9535441399f, 0.2338332087f),
+	FVector2D(0.9811842442f, 0.3334831297f),
+	FVector2D(0.9811842442f, 0.3334831297f),
+	FVector2D(1.0f, 0.4600457847f),
+	FVector2D(1.0f, 0.6219738722f),
+	FVector2D(1.0f, 0.6219738722f),
+	FVector2D(0.9842697382f, 0.7411668301f),
+	FVector2D(0.9689071178f, 0.8031471968f),
+	FVector2D(0.8285242319f, 0.9551845193f),
+	FVector2D(0.7284021378f, 0.9779636264f),
+	FVector2D(0.6235122681f, 0.9927965403f),
+	FVector2D(0.0734243989f, 0.131273821f),
+	FVector2D(0.1512531638f, 0.03671549633f),
+	FVector2D(0.3018191457f, 0.00689325761f),
+	FVector2D(0.04287481308f, 0.216376245f),
+	FVector2D(0.0734243989f, 0.131273821f),
+	FVector2D(0.3018191457f, 0.00689325761f),
+	FVector2D(0.0f, 0.4462257326f),
+	FVector2D(0.04287481308f, 0.216376245f),
+	FVector2D(0.3018191457f, 0.00689325761f),
+	FVector2D(0.0f, 0.4462257326f),
+	FVector2D(0.3018191457f, 0.00689325761f),
+	FVector2D(0.5447612405f, 0.0f),
+	FVector2D(0.006614744663f, 0.6894935966f),
+	FVector2D(0.0f, 0.4462257326f),
+	FVector2D(0.5447612405f, 0.0f),
+	FVector2D(0.006614744663f, 0.6894935966f),
+	FVector2D(0.5447612405f, 0.0f),
+	FVector2D(0.7506074309f, 0.03307863697f),
+	FVector2D(0.04205995798f, 0.8690519333f),
+	FVector2D(0.006614744663f, 0.6894935966f),
+	FVector2D(0.7506074309f, 0.03307863697f),
+	FVector2D(0.1185472012f, 0.9627953768f),
+	FVector2D(0.04205995798f, 0.8690519333f),
+	FVector2D(0.7506074309f, 0.03307863697f),
+	FVector2D(0.1740469933f, 0.9753877521f),
+	FVector2D(0.1185472012f, 0.9627953768f),
+	FVector2D(0.7506074309f, 0.03307863697f),
+	FVector2D(0.3806557059f, 1.0f),
+	FVector2D(0.1740469933f, 0.9753877521f),
+	FVector2D(0.7506074309f, 0.03307863697f),
+	FVector2D(0.3806557059f, 1.0f),
+	FVector2D(0.7506074309f, 0.03307863697f),
+	FVector2D(0.6235122681f, 0.9927965403f),
+	FVector2D(0.7506074309f, 0.03307863697f),
+	FVector2D(0.8285242319f, 0.9551845193f),
+	FVector2D(0.6235122681f, 0.9927965403f),
+	FVector2D(0.7506074309f, 0.03307863697f),
+	FVector2D(0.9503660202f, 0.8630086184f),
+	FVector2D(0.8285242319f, 0.9551845193f),
+	FVector2D(0.9689071178f, 0.8031471968f),
+	FVector2D(0.7506074309f, 0.03307863697f),
+	FVector2D(0.9503660202f, 0.8630086184f),
+	FVector2D(0.7506074309f, 0.03307863697f),
+	FVector2D(1.0f, 0.6219738722f),
+	FVector2D(0.9689071178f, 0.8031471968f),
+	FVector2D(0.7506074309f, 0.03307863697f),
+	FVector2D(0.9811842442f, 0.3334831297f),
+	FVector2D(1.0f, 0.6219738722f),
+	FVector2D(0.7506074309f, 0.03307863697f),
+	FVector2D(0.8197077513f, 0.05489978567f),
+	FVector2D(0.9811842442f, 0.3334831297f),
+	FVector2D(0.8197077513f, 0.05489978567f),
+	FVector2D(0.9157208204f, 0.1436391324f),
+	FVector2D(0.9811842442f, 0.3334831297f),
+};
diff --git a/Engine/Plugins/Runtime/OculusRift/Source/OculusRift/Private/OculusRiftRender.cpp b/Engine/Plugins/Runtime/OculusRift/Source/OculusRift/Private/OculusRiftRender.cpp
--- a/Engine/Plugins/Runtime/OculusRift/Source/OculusRift/Private/OculusRiftRender.cpp
+++ b/Engine/Plugins/Runtime/OculusRift/Source/OculusRift/Private/OculusRiftRender.cpp
@@ -110,6 +110,38 @@ void FOculusRiftHMD::RenderTexture_RenderThread(class FRHICommandListImmediate&
 	}
 }
 
+static void DrawOcclusionMesh(FRHICommandList& RHICmdList, EStereoscopicPass StereoPass, const FHMDViewMesh MeshAssets[])
+{
+	check(IsInRenderingThread());
+	check(StereoPass != eSSP_FULL);
+
+	const uint32 MeshIndex = (StereoPass == eSSP_LEFT_EYE) ? 0 : 1;
+	const FHMDViewMesh& Mesh = MeshAssets[MeshIndex];
+	check(Mesh.IsValid());
+
+	DrawIndexedPrimitiveUP(
+		RHICmdList,
+		PT_TriangleList,
+		0,
+		Mesh.NumVertices,
+		Mesh.NumTriangles,
+		Mesh.pIndices,
+		sizeof(Mesh.pIndices[0]),
+		Mesh.pVertices,
+		sizeof(Mesh.pVertices[0])
+		);
+}
+
+void FOculusRiftHMD::DrawHiddenAreaMesh_RenderThread(FRHICommandList& RHICmdList, EStereoscopicPass StereoPass) const
+{
+	DrawOcclusionMesh(RHICmdList, StereoPass, HiddenAreaMeshes);
+}
+
+void FOculusRiftHMD::DrawVisibleAreaMesh_RenderThread(FRHICommandList& RHICmdList, EStereoscopicPass StereoPass) const
+{
+	DrawOcclusionMesh(RHICmdList, StereoPass, VisibleAreaMeshes);
+}
+
 void FViewExtension::PreRenderViewFamily_RenderThread(FRHICommandListImmediate& RHICmdList, FSceneViewFamily& ViewFamily)
 {
 	check(IsInRenderingThread());
diff --git a/Engine/Plugins/Runtime/Steam/SteamVR/Source/SteamVR/Private/SteamVRHMD.cpp b/Engine/Plugins/Runtime/Steam/SteamVR/Source/SteamVR/Private/SteamVRHMD.cpp
--- a/Engine/Plugins/Runtime/Steam/SteamVR/Source/SteamVR/Private/SteamVRHMD.cpp
+++ b/Engine/Plugins/Runtime/Steam/SteamVR/Source/SteamVR/Private/SteamVRHMD.cpp
@@ -9,7 +9,7 @@
 #include "PostProcess/PostProcessHMD.h"
 #include "Classes/SteamVRFunctionLibrary.h"
 
-#include "ViveVisibleAreaMesh.h"
+#include "SteamVRMeshAssets.h"
 
 #if WITH_EDITOR
 #include "Editor/UnrealEd/Classes/Editor/EditorEngine.h"
@@ -994,11 +994,7 @@ void FSteamVRHMD::Startup()
 			}
 		}
 
-		// Setup meshes
-		HiddenAreaMeshes[0].Build(VRSystem->GetHiddenAreaMesh(vr::Hmd_Eye::Eye_Left));
-		HiddenAreaMeshes[1].Build(VRSystem->GetHiddenAreaMesh(vr::Hmd_Eye::Eye_Right));
-		VisibleAreaMeshes[0].Build(LeftEyePositions, VertexCount);
-		VisibleAreaMeshes[1].Build(RightEyePositions, VertexCount);
+		SetupOcclusionMeshes();
 
 #if PLATFORM_WINDOWS
 		if (IsPCPlatform(GMaxRHIShaderPlatform) && !IsOpenGLPlatform(GMaxRHIShaderPlatform))
@@ -1119,117 +1115,57 @@ void FSteamVRHMD::UnloadOpenVRModule()
 	}
 }
 
-FSteamVRHMD::FHiddenAreaMesh::FHiddenAreaMesh() :
-	pVertices(nullptr),
-	pIndices(nullptr),
-	NumVertices(0),
-	NumIndices(0),
-	NumTriangles(0)
-{}
-
-FSteamVRHMD::FHiddenAreaMesh::~FHiddenAreaMesh()
-{
-	if (pVertices)
-	{
-		delete[] pVertices;
-	}
-
-	if (pIndices)
-	{
-		delete[] pIndices;
-	}
-}
-
-void FSteamVRHMD::FHiddenAreaMesh::Build(const vr::HiddenAreaMesh_t& Mesh)
-{
-	check(pVertices == nullptr);
-
-	NumTriangles = Mesh.unTriangleCount;
-	NumVertices  = NumTriangles * 3;
-	NumIndices   = NumVertices;
+void FSteamVRHMD::SetupOcclusionMeshes()
+{	
+	const vr::HiddenAreaMesh_t LeftEyeMesh = VRSystem->GetHiddenAreaMesh(vr::Hmd_Eye::Eye_Left);
+	const vr::HiddenAreaMesh_t RightEyeMesh = VRSystem->GetHiddenAreaMesh(vr::Hmd_Eye::Eye_Right);
+	
+	const uint32 VertexCount = LeftEyeMesh.unTriangleCount * 3;
+	check(LeftEyeMesh.unTriangleCount == RightEyeMesh.unTriangleCount);
 
-	// Did we get any data?
-	if (NumVertices == 0)
+	// Copy mesh data from SteamVR format to ours, then initialize the meshes.
+	if (VertexCount > 0)
 	{
-		return;
-	}
+		FVector2D* const LeftEyePositions = new FVector2D[VertexCount];
+		FVector2D* const RightEyePositions = new FVector2D[VertexCount];
 
-	pVertices = new FVector4[NumVertices];
-	pIndices  = new uint16[NumIndices];
-
-	uint32 DataIndex = 0;
-	for (uint32 TriangleIter = 0; TriangleIter < NumTriangles; ++TriangleIter)
-	{
-		for (uint32 VertexIter = 0; VertexIter < 3; ++VertexIter)
+		uint32 HiddenAreaMeshCrc = 0;
+		uint32 DataIndex = 0;
+		for (uint32 TriangleIter = 0; TriangleIter < LeftEyeMesh.unTriangleCount; ++TriangleIter)
 		{
-			const vr::HmdVector2_t& Vertex = Mesh.pVertexData[DataIndex];
-
-			// Remap from [0 1] to [-1 1] to match our NDC space
-			pVertices[DataIndex].X = (Vertex.v[0] * 2.0f) - 1.0f;
-			pVertices[DataIndex].Y = (Vertex.v[1] * 2.0f) - 1.0f;
-			pVertices[DataIndex].Z = 1.0f; // Setting to 1 for reversed depth near plane
-			pVertices[DataIndex].W = 1.0f; 
-
-			pIndices[DataIndex] = DataIndex;
+			for (uint32 VertexIter = 0; VertexIter < 3; ++VertexIter)
+			{
+				const vr::HmdVector2_t& LeftSrc = LeftEyeMesh.pVertexData[DataIndex];
+				const vr::HmdVector2_t& RightSrc = RightEyeMesh.pVertexData[DataIndex];
 
-			++DataIndex;
-		}
-	}
-}
+				FVector2D& LeftDst = LeftEyePositions[DataIndex];
+				FVector2D& RightDst = RightEyePositions[DataIndex];
 
-FSteamVRHMD::FVisibleAreaMesh::FVisibleAreaMesh() :
-	pVertices(nullptr),
-	pIndices(nullptr),
-	NumVertices(0),
-	NumIndices(0),
-	NumTriangles(0)
-{}
+				LeftDst.X = LeftSrc.v[0];
+				LeftDst.Y = LeftSrc.v[1];
 
-FSteamVRHMD::FVisibleAreaMesh::~FVisibleAreaMesh()
-{
-	if (pVertices)
-	{
-		delete[] pVertices;
-	}
+				RightDst.X = RightSrc.v[0];
+				RightDst.Y = RightSrc.v[1];
 
-	if (pIndices)
-	{
-		delete[] pIndices;
-	}
-}
-
-void FSteamVRHMD::FVisibleAreaMesh::Build(const MeshVertex Positions[], uint32 InNumVertices)
-{
-	check(pVertices == nullptr);
-	check(InNumVertices > 2 && InNumVertices % 3 == 0);
+				HiddenAreaMeshCrc = FCrc::MemCrc32(&LeftDst, sizeof(FVector2D), HiddenAreaMeshCrc);
 
-	NumVertices = InNumVertices;
-	NumTriangles = NumVertices / 3;
-	NumIndices = NumVertices;
+				++DataIndex;
+			}
+		}
 
-	pVertices = new FFilterVertex[NumVertices];
-	pIndices = new uint16[NumIndices];
+		HiddenAreaMeshes[0].BuildMesh(LeftEyePositions, VertexCount, FHMDViewMesh::MT_HiddenArea);
+		HiddenAreaMeshes[1].BuildMesh(RightEyePositions, VertexCount, FHMDViewMesh::MT_HiddenArea);
 
-	uint32 DataIndex = 0;
-	for (uint32 TriangleIter = 0; TriangleIter < NumTriangles; ++TriangleIter)
-	{
-		for (uint32 VertexIter = 0; VertexIter < 3; ++VertexIter)
+		// If the hidden area mesh from the SteamVR runtime matches the mesh used to generate the Vive's visible area mesh, initialize it.
+		if (HiddenAreaMeshCrc == ViveHiddenAreaMeshCrc)
 		{
-			const MeshVertex& Position = Positions[DataIndex];
-			FFilterVertex& Vertex = pVertices[DataIndex];
-			
-			Vertex.Position.X = Position.X;
-			Vertex.Position.Y = 1.0f - Position.Y;
-			Vertex.Position.Z = 0.0f;
-			Vertex.Position.W = 1.0f;
-			
-			Vertex.UV.X = Position.X;
-			Vertex.UV.Y = 1.0f - Position.Y;
-			
-			pIndices[DataIndex] = DataIndex;
-			
-			++DataIndex;
+			VisibleAreaMeshes[0].BuildMesh(Vive_LeftEyeVisibleAreaPositions, VisibleAreaVertexCount, FHMDViewMesh::MT_VisibleArea);
+			VisibleAreaMeshes[1].BuildMesh(Vive_RightEyeVisibleAreaPositions, VisibleAreaVertexCount, FHMDViewMesh::MT_VisibleArea);
 		}
+
+		delete[] LeftEyePositions;
+		delete[] RightEyePositions;
 	}
 }
+
 #endif //STEAMVR_SUPPORTED_PLATFORMS
diff --git a/Engine/Plugins/Runtime/Steam/SteamVR/Source/SteamVR/Private/SteamVRHMD.h b/Engine/Plugins/Runtime/Steam/SteamVR/Source/SteamVR/Private/SteamVRHMD.h
--- a/Engine/Plugins/Runtime/Steam/SteamVR/Source/SteamVR/Private/SteamVRHMD.h
+++ b/Engine/Plugins/Runtime/Steam/SteamVR/Source/SteamVR/Private/SteamVRHMD.h
@@ -84,11 +84,11 @@ public:
 	virtual void SetBaseOrientation(const FQuat& BaseOrient) override;
 	virtual FQuat GetBaseOrientation() const override;
 
-	virtual bool HasHiddenAreaMask() const override { return HiddenAreaMeshes[0].IsValid() && HiddenAreaMeshes[1].IsValid(); }
-	virtual void DrawHiddenAreaMaskView_RenderThread(FRHICommandList& RHICmdList, EStereoscopicPass StereoPass) const override;
+	virtual bool HasHiddenAreaMesh() const override { return HiddenAreaMeshes[0].IsValid() && HiddenAreaMeshes[1].IsValid(); }
+	virtual void DrawHiddenAreaMesh_RenderThread(FRHICommandList& RHICmdList, EStereoscopicPass StereoPass) const override;
 
-	virtual bool HasPostProcessMesh() const override { return VisibleAreaMeshes[0].IsValid() && VisibleAreaMeshes[1].IsValid(); }
-	virtual void DrawPostProcessMesh_RenderThread(FRHICommandList& RHICmdList, EStereoscopicPass StereoPass) const override;
+	virtual bool HasVisibleAreaMesh() const override { return VisibleAreaMeshes[0].IsValid() && VisibleAreaMeshes[1].IsValid(); }
+	virtual void DrawVisibleAreaMesh_RenderThread(FRHICommandList& RHICmdList, EStereoscopicPass StereoPass) const override;
 
 	virtual void DrawDistortionMesh_RenderThread(struct FRenderingCompositePassContext& Context, const FIntPoint& TextureSize) override;
 
@@ -253,6 +253,8 @@ private:
 
 private:
 
+	void SetupOcclusionMeshes();
+
 	bool bHmdEnabled;
 	bool bStereoEnabled;
 	bool bHmdPosTracking;
@@ -300,46 +302,8 @@ private:
 		return FVector(InVector.v[0], InVector.v[1], InVector.v[2]);
 	}
 
-	struct FHiddenAreaMesh
-	{
-		FHiddenAreaMesh();
-		~FHiddenAreaMesh();
-
-		bool IsValid() const
-		{
-			return NumTriangles > 0;
-		}
-
-		void Build(const vr::HiddenAreaMesh_t& Mesh);
-
-		FVector4* pVertices;
-		uint16*   pIndices;
-		unsigned  NumVertices;
-		unsigned  NumIndices;
-		unsigned  NumTriangles;
-	};
-
-	struct FVisibleAreaMesh
-	{
-		FVisibleAreaMesh();
-		~FVisibleAreaMesh();
-
-		bool IsValid() const
-		{
-			return NumTriangles > 0;
-		}
-
-		void Build(const struct MeshVertex Positions[], uint32 InNumVertices);
-
-		FFilterVertex* pVertices;
-		uint16*   pIndices;
-		unsigned  NumVertices;
-		unsigned  NumIndices;
-		unsigned  NumTriangles;
-	};
-
-	FHiddenAreaMesh HiddenAreaMeshes[2];
-	FVisibleAreaMesh VisibleAreaMeshes[2];
+	FHMDViewMesh HiddenAreaMeshes[2];
+	FHMDViewMesh VisibleAreaMeshes[2];
 
 	/** Chaperone Support */
 	struct FChaperoneBounds
diff --git a/Engine/Plugins/Runtime/Steam/SteamVR/Source/SteamVR/Private/SteamVRMeshAssets.h b/Engine/Plugins/Runtime/Steam/SteamVR/Source/SteamVR/Private/SteamVRMeshAssets.h
new file mode 100644
--- /dev/null
+++ b/Engine/Plugins/Runtime/Steam/SteamVR/Source/SteamVR/Private/SteamVRMeshAssets.h
@@ -0,0 +1,153 @@
+// Copyright 1998-2015 Epic Games, Inc. All Rights Reserved.
+
+#pragma once
+
+// These are the inverse of the hidden area meshes Valve provides us for the Vive.
+// i.e. a mesh that covers all pixels which *are* inputs to the distortion pass.
+static const uint32 ViveHiddenAreaMeshCrc = 2526529894;
+static const uint32 VisibleAreaVertexCount = 69;
+static const FVector2D Vive_LeftEyeVisibleAreaPositions[] =
+{
+	FVector2D(1.0f, 0.1999999881f),
+	FVector2D(1.0f, 0.1949999332f),
+	FVector2D(0.9099999552f, 0.1099999547f),
+	FVector2D(0.9159999971f, 0.3639999628f),
+	FVector2D(1.0f, 0.1999999881f),
+	FVector2D(0.9099999552f, 0.1099999547f),
+	FVector2D(0.9099999552f, 0.1099999547f),
+	FVector2D(0.9159999971f, 0.3639999628f),
+	FVector2D(0.8100000505f, 0.04499995708f),
+	FVector2D(0.8100000505f, 0.04499995708f),
+	FVector2D(0.9159999971f, 0.3639999628f),
+	FVector2D(0.6799999361f, 0.0f),
+	FVector2D(0.6799999361f, 0.0f),
+	FVector2D(0.9159999971f, 0.3639999628f),
+	FVector2D(0.3699999933f, 0.0f),
+	FVector2D(0.8949999695f, 0.4949999452f),
+	FVector2D(0.9159999971f, 0.3639999628f),
+	FVector2D(0.3699999933f, 0.0f),
+	FVector2D(0.3699999933f, 0.0f),
+	FVector2D(0.8949999695f, 0.4949999452f),
+	FVector2D(0.2399999981f, 0.04999995232f),
+	FVector2D(0.2399999981f, 0.04999995232f),
+	FVector2D(0.8949999695f, 0.4949999452f),
+	FVector2D(0.1320000176f, 0.1229999661f),
+	FVector2D(0.1320000176f, 0.1229999661f),
+	FVector2D(0.8949999695f, 0.4949999452f),
+	FVector2D(0.05999999094f, 0.1999999881f),
+	FVector2D(0.05999999094f, 0.1999999881f),
+	FVector2D(0.8949999695f, 0.4949999452f),
+	FVector2D(0.0f, 0.2949999571f),
+	FVector2D(0.0f, 0.2949999571f),
+	FVector2D(0.8949999695f, 0.4949999452f),
+	FVector2D(0.0f, 0.6999999285f),
+	FVector2D(0.0f, 0.6999999285f),
+	FVector2D(0.8949999695f, 0.4949999452f),
+	FVector2D(0.04000001001f, 0.7680000067f),
+	FVector2D(0.04000001001f, 0.7680000067f),
+	FVector2D(0.8949999695f, 0.4949999452f),
+	FVector2D(0.09000002193f, 0.8329999447f),
+	FVector2D(0.09000002193f, 0.8329999447f),
+	FVector2D(0.8949999695f, 0.4949999452f),
+	FVector2D(0.1899999862f, 0.9149999619f),
+	FVector2D(0.1899999862f, 0.9149999619f),
+	FVector2D(0.8949999695f, 0.4949999452f),
+	FVector2D(0.2710000162f, 0.9599999189f),
+	FVector2D(0.2710000162f, 0.9599999189f),
+	FVector2D(0.8949999695f, 0.4949999452f),
+	FVector2D(0.388f, 1.0f),
+	FVector2D(0.388f, 1.0f),
+	FVector2D(0.8949999695f, 0.4950000048f),
+	FVector2D(0.6560000067f, 1.0f),
+	FVector2D(0.9129999762f, 0.628000021f),
+	FVector2D(0.8949999695f, 0.4949999452f),
+	FVector2D(0.6560000067f, 1.0f),
+	FVector2D(0.6560000067f, 1.0f),
+	FVector2D(0.9129999762f, 0.628000021f),
+	FVector2D(0.7399999981f, 0.9749999046f),
+	FVector2D(0.7399999981f, 0.9749999046f),
+	FVector2D(0.9129999762f, 0.628000021f),
+	FVector2D(0.8129999523f, 0.9449999332f),
+	FVector2D(0.8129999523f, 0.9449999332f),
+	FVector2D(0.9129999762f, 0.628000021f),
+	FVector2D(0.9299999361f, 0.8679999113f),
+	FVector2D(1.0f, 0.7999999523f),
+	FVector2D(0.9299999361f, 0.8679999113f),
+	FVector2D(1.0f, 0.7879999876f),
+	FVector2D(1.0f, 0.7879999876f),
+	FVector2D(0.9129999762f, 0.628000021f),
+	FVector2D(0.9299999361f, 0.8679999113f),
+};
+
+static const FVector2D Vive_RightEyeVisibleAreaPositions[] =
+{
+	FVector2D(0.09000000358f, 0.1099999917f),
+	FVector2D(0.0f, 0.195f),
+	FVector2D(0.0f, 0.1999999952f),
+	FVector2D(0.0f, 0.1999999952f),
+	FVector2D(0.09000000358f, 0.1099999917f),
+	FVector2D(0.08399999142f, 0.3639999998f),
+	FVector2D(0.1899999976f, 0.04499999404f),
+	FVector2D(0.09000000358f, 0.1099999917f),
+	FVector2D(0.08399999142f, 0.3639999998f),
+	FVector2D(0.3199999928f, 0.0f),
+	FVector2D(0.1899999976f, 0.04499999404f),
+	FVector2D(0.08399999142f, 0.3639999998f),
+	FVector2D(0.6299999952f, 0.0f),
+	FVector2D(0.3199999928f, 0.0f),
+	FVector2D(0.08399999142f, 0.3639999998f),
+	FVector2D(0.08399999142f, 0.3639999998f),
+	FVector2D(0.6299999952f, 0.0f),
+	FVector2D(0.1049999893f, 0.4949999821f),
+	FVector2D(0.7599999905f, 0.04999998927f),
+	FVector2D(0.6299999952f, 0.0f),
+	FVector2D(0.1049999893f, 0.4949999821f),
+	FVector2D(0.8680000305f, 0.1230000031f),
+	FVector2D(0.7599999905f, 0.04999998927f),
+	FVector2D(0.1049999893f, 0.4949999821f),
+	FVector2D(0.9400000572f, 0.1999999952f),
+	FVector2D(0.8680000305f, 0.1230000031f),
+	FVector2D(0.1049999893f, 0.4949999821f),
+	FVector2D(1.0f, 0.294999994f),
+	FVector2D(0.9400000572f, 0.1999999952f),
+	FVector2D(0.1049999893f, 0.4949999821f),
+	FVector2D(1.0f, 0.700000025f),
+	FVector2D(1.0f, 0.294999994f),
+	FVector2D(0.1049999893f, 0.4949999821f),
+	FVector2D(0.9600000381f, 0.767999984f),
+	FVector2D(1.0f, 0.700000025f),
+	FVector2D(0.1049999893f, 0.4949999821f),
+	FVector2D(0.9099999666f, 0.8330000412f),
+	FVector2D(0.9600000381f, 0.767999984f),
+	FVector2D(0.1049999893f, 0.4949999821f),
+	FVector2D(0.8099999428f, 0.9150000584f),
+	FVector2D(0.9099999666f, 0.8330000412f),
+	FVector2D(0.1049999893f, 0.4949999821f),
+	FVector2D(0.7289999723f, 0.9600000155f),
+	FVector2D(0.8099999428f, 0.9150000584f),
+	FVector2D(0.1049999893f, 0.4949999821f),
+	FVector2D(0.6119999886f, 1.0f),
+	FVector2D(0.7289999723f, 0.9600000155f),
+	FVector2D(0.1049999893f, 0.4949999821f),
+	FVector2D(0.3439999819f, 1.0f),
+	FVector2D(0.6119999886f, 1.0f),
+	FVector2D(0.1049999893f, 0.4949999821f),
+	FVector2D(0.1049999893f, 0.4949999821f),
+	FVector2D(0.3439999819f, 1.0f),
+	FVector2D(0.0870000124f, 0.6279999983f),
+	FVector2D(0.2599999905f, 0.9750000012f),
+	FVector2D(0.3439999819f, 1.0f),
+	FVector2D(0.0870000124f, 0.6279999983f),
+	FVector2D(0.1870000064f, 0.9450000298f),
+	FVector2D(0.2599999905f, 0.9750000012f),
+	FVector2D(0.0870000124f, 0.6279999983f),
+	FVector2D(0.06999999285f, 0.8680000079f),
+	FVector2D(0.1870000064f, 0.9450000298f),
+	FVector2D(0.0870000124f, 0.6279999983f),
+	FVector2D(0.0f, 0.8000000489f),
+	FVector2D(0.0f, 0.787999965f),
+	FVector2D(0.06999999285f, 0.8680000079f),
+	FVector2D(0.0870000124f, 0.6279999983f),
+	FVector2D(0.06999999285f, 0.8680000079f),
+	FVector2D(0.0f, 0.787999965f),
+};
diff --git a/Engine/Plugins/Runtime/Steam/SteamVR/Source/SteamVR/Private/SteamVRRender.cpp b/Engine/Plugins/Runtime/Steam/SteamVR/Source/SteamVR/Private/SteamVRRender.cpp
--- a/Engine/Plugins/Runtime/Steam/SteamVR/Source/SteamVR/Private/SteamVRRender.cpp
+++ b/Engine/Plugins/Runtime/Steam/SteamVR/Source/SteamVR/Private/SteamVRRender.cpp
@@ -75,13 +75,13 @@ void FSteamVRHMD::RenderTexture_RenderThread(FRHICommandListImmediate& RHICmdLis
 	}
 }
 
-void FSteamVRHMD::DrawHiddenAreaMaskView_RenderThread(FRHICommandList& RHICmdList, EStereoscopicPass StereoPass) const
+static void DrawOcclusionMesh(FRHICommandList& RHICmdList, EStereoscopicPass StereoPass, const FHMDViewMesh MeshAssets[])
 {
 	check(IsInRenderingThread());
 	check(StereoPass != eSSP_FULL);
 
-	const uint32 HiddenMeshIndex = (StereoPass == eSSP_LEFT_EYE) ? 0 : 1;
-	const FHiddenAreaMesh& Mesh = HiddenAreaMeshes[HiddenMeshIndex];
+	const uint32 MeshIndex = (StereoPass == eSSP_LEFT_EYE) ? 0 : 1;
+	const FHMDViewMesh& Mesh = MeshAssets[MeshIndex];
 	check(Mesh.IsValid());
 
 	DrawIndexedPrimitiveUP(
@@ -94,29 +94,17 @@ void FSteamVRHMD::DrawHiddenAreaMaskView_RenderThread(FRHICommandList& RHICmdLis
 		sizeof(Mesh.pIndices[0]),
 		Mesh.pVertices,
 		sizeof(Mesh.pVertices[0])
-	);
+		);
 }
 
-void FSteamVRHMD::DrawPostProcessMesh_RenderThread(FRHICommandList& RHICmdList, EStereoscopicPass StereoPass) const
+void FSteamVRHMD::DrawHiddenAreaMesh_RenderThread(FRHICommandList& RHICmdList, EStereoscopicPass StereoPass) const
 {
-	check(IsInRenderingThread());
-	check(StereoPass != eSSP_FULL);
-
-	const uint32 HiddenMeshIndex = (StereoPass == eSSP_LEFT_EYE) ? 0 : 1;
-	const FVisibleAreaMesh& Mesh = VisibleAreaMeshes[HiddenMeshIndex];
-	check(Mesh.IsValid());
+	DrawOcclusionMesh(RHICmdList, StereoPass, HiddenAreaMeshes);
+}
 
-	DrawIndexedPrimitiveUP(
-		RHICmdList,
-		PT_TriangleList,
-		0,
-		Mesh.NumVertices,
-		Mesh.NumTriangles,
-		Mesh.pIndices,
-		sizeof(Mesh.pIndices[0]),
-		Mesh.pVertices,
-		sizeof(Mesh.pVertices[0])
-		);
+void FSteamVRHMD::DrawVisibleAreaMesh_RenderThread(FRHICommandList& RHICmdList, EStereoscopicPass StereoPass) const
+{
+	DrawOcclusionMesh(RHICmdList, StereoPass, VisibleAreaMeshes);
 }
 
 #if PLATFORM_WINDOWS
diff --git a/Engine/Plugins/Runtime/Steam/SteamVR/Source/SteamVR/Private/ViveVisibleAreaMesh.h b/Engine/Plugins/Runtime/Steam/SteamVR/Source/SteamVR/Private/ViveVisibleAreaMesh.h
deleted file mode 100644
diff --git a/Engine/Source/Runtime/HeadMountedDisplay/Private/HeadMountedDisplayTypes.cpp b/Engine/Source/Runtime/HeadMountedDisplay/Private/HeadMountedDisplayTypes.cpp
new file mode 100644
--- /dev/null
+++ b/Engine/Source/Runtime/HeadMountedDisplay/Private/HeadMountedDisplayTypes.cpp
@@ -0,0 +1,73 @@
+#include "HeadMountedDisplayPrivate.h"
+
+FHMDViewMesh::FHMDViewMesh() :
+	pVertices(nullptr),
+	pIndices(nullptr),
+	NumVertices(0),
+	NumIndices(0),
+	NumTriangles(0)
+{}
+
+FHMDViewMesh::~FHMDViewMesh()
+{
+	if (pVertices)
+	{
+		delete[] pVertices;
+	}
+
+	if (pIndices)
+	{
+		delete[] pIndices;
+	}
+}
+
+void FHMDViewMesh::BuildMesh(const FVector2D Positions[], uint32 VertexCount, EHMDMeshType MeshType)
+{
+	check(pVertices == nullptr);
+	check(VertexCount > 2 && VertexCount % 3 == 0);
+
+	NumVertices = VertexCount;
+	NumTriangles = NumVertices / 3;
+	NumIndices = NumVertices;
+
+	pVertices = new FFilterVertex[NumVertices];
+	pIndices = new uint16[NumIndices];
+
+	uint32 DataIndex = 0;
+	for (uint32 TriangleIter = 0; TriangleIter < NumTriangles; ++TriangleIter)
+	{
+		for (uint32 VertexIter = 0; VertexIter < 3; ++VertexIter)
+		{
+			const FVector2D& Position = Positions[DataIndex];
+			FFilterVertex& Vertex = pVertices[DataIndex];
+
+			if (MeshType == MT_HiddenArea)
+			{
+				// Remap from to NDC space [0 1] -> [-1 1]
+				Vertex.Position.X = (Position.X * 2.0f) - 1.0f;
+				Vertex.Position.Y = (Position.Y * 2.0f) - 1.0f;
+				Vertex.Position.Z = 1.0f;
+				Vertex.Position.W = 1.0f;
+
+				// Not used for hidden area
+				Vertex.UV.X = 0.0f;
+				Vertex.UV.Y = 0.0f;
+			}
+			else
+			{
+				// Remap the viewport origin from the bottom left to the top left
+				Vertex.Position.X = Position.X;
+				Vertex.Position.Y = 1.0f - Position.Y;
+				Vertex.Position.Z = 0.0f;
+				Vertex.Position.W = 1.0f;
+
+				Vertex.UV.X = Position.X;
+				Vertex.UV.Y = 1.0f - Position.Y;
+			}
+
+			pIndices[DataIndex] = DataIndex;
+
+			++DataIndex;
+		}
+	}
+}
\ No newline at end of file
diff --git a/Engine/Source/Runtime/HeadMountedDisplay/Public/HeadMountedDisplayTypes.h b/Engine/Source/Runtime/HeadMountedDisplay/Public/HeadMountedDisplayTypes.h
new file mode 100644
--- /dev/null
+++ b/Engine/Source/Runtime/HeadMountedDisplay/Public/HeadMountedDisplayTypes.h
@@ -0,0 +1,45 @@
+// Copyright 1998-2015 Epic Games, Inc. All Rights Reserved.
+
+#pragma once
+
+/**
+* The family of HMD device.  Register a new class of device here if you need to branch code for PostProcessing until
+*/
+namespace EHMDDeviceType
+{
+	enum Type
+	{
+		DT_OculusRift,
+		DT_Morpheus,
+		DT_ES2GenericStereoMesh,
+		DT_SteamVR,
+		DT_GearVR
+	};
+}
+
+class HEADMOUNTEDDISPLAY_API FHMDViewMesh
+{
+public:
+
+	enum EHMDMeshType
+	{
+		MT_HiddenArea,
+		MT_VisibleArea
+	};
+
+	FHMDViewMesh();
+	~FHMDViewMesh();
+
+	bool IsValid() const
+	{
+		return NumTriangles > 0;
+	}
+
+	void BuildMesh(const FVector2D Positions[], uint32 VertexCount, EHMDMeshType MeshType);
+
+	FFilterVertex* pVertices;
+	uint16*   pIndices;
+	unsigned  NumVertices;
+	unsigned  NumIndices;
+	unsigned  NumTriangles;
+};
diff --git a/Engine/Source/Runtime/HeadMountedDisplay/Public/IHeadMountedDisplay.h b/Engine/Source/Runtime/HeadMountedDisplay/Public/IHeadMountedDisplay.h
--- a/Engine/Source/Runtime/HeadMountedDisplay/Public/IHeadMountedDisplay.h
+++ b/Engine/Source/Runtime/HeadMountedDisplay/Public/IHeadMountedDisplay.h
@@ -4,25 +4,11 @@
 
 #include "StereoRendering.h"
 #include "Layout/SlateRect.h"
+#include "HeadMountedDisplayTypes.h"
 
 // depending on your kit and SDK, you may want to use this.
 // new distortion handling still in development.
 
-/**
- * The family of HMD device.  Register a new class of device here if you need to branch code for PostProcessing until 
- */
-namespace EHMDDeviceType
-{
-	enum Type
-	{
-		DT_OculusRift,
-		DT_Morpheus,
-		DT_ES2GenericStereoMesh,
-		DT_SteamVR,
-		DT_GearVR
-	};
-}
-
 /**
  * HMD device interface
  */
@@ -261,26 +247,26 @@ public:
 	virtual FVector GetPositionScale3D() const { return FVector::ZeroVector; }
 
 	/**
-	* @return true if a hidden area mask is available for the device.
+	* @return true if a hidden area mesh is available for the device.
 	*/
-	virtual bool HasHiddenAreaMask() const { return false; }
+	virtual bool HasHiddenAreaMesh() const { return false; }
 
 	/**
-	* @return true if a post process mesh is available for the device.
+	* @return true if a visible area mesh is available for the device.
 	*/
-	virtual bool HasPostProcessMesh() const { return false; }
+	virtual bool HasVisibleAreaMesh() const { return false; }
 
 	/**
-	* Optional method to draw a view's hidden area mask where supported.
+	* Optional method to draw a view's hidden area mesh where supported.
 	* This can be used to avoid rendering pixels which are not included as input into the final distortion pass.
 	*/
-	virtual void DrawHiddenAreaMaskView_RenderThread(class FRHICommandList& RHICmdList, EStereoscopicPass StereoPass) const {};
+	virtual void DrawHiddenAreaMesh_RenderThread(class FRHICommandList& RHICmdList, EStereoscopicPass StereoPass) const {};
 
 	/**
-	* Optional method to draw a view's post process mesh where supported.
+	* Optional method to draw a view's visible area mesh where supported.
 	* This can be used instead of a full screen quad to avoid rendering pixels which are not included as input into the final distortion pass.
 	*/
-	virtual void DrawPostProcessMesh_RenderThread(class FRHICommandList& RHICmdList, EStereoscopicPass StereoPass) const {};
+	virtual void DrawVisibleAreaMesh_RenderThread(class FRHICommandList& RHICmdList, EStereoscopicPass StereoPass) const {};
 
 	virtual void DrawDistortionMesh_RenderThread(struct FRenderingCompositePassContext& Context, const FIntPoint& TextureSize) {}
 
diff --git a/Engine/Source/Runtime/Renderer/Private/DeferredShadingRenderer.cpp b/Engine/Source/Runtime/Renderer/Private/DeferredShadingRenderer.cpp
--- a/Engine/Source/Runtime/Renderer/Private/DeferredShadingRenderer.cpp
+++ b/Engine/Source/Runtime/Renderer/Private/DeferredShadingRenderer.cpp
@@ -697,7 +697,7 @@ static FORCEINLINE bool HasHiddenAreaMask()
 		HiddenAreaMaskCVar->GetValueOnRenderThread() == 1 &&
 		GEngine &&
 		GEngine->HMDDevice.IsValid() &&
-		GEngine->HMDDevice->HasHiddenAreaMask());
+		GEngine->HMDDevice->HasHiddenAreaMesh());
 }
 
 static void SetAndClearViewGBuffer(FRHICommandListImmediate& RHICmdList, FViewInfo& View, bool bClearDepth)
@@ -1429,7 +1429,7 @@ static void RenderHiddenAreaMaskView(FRHICommandList& RHICmdList, const FViewInf
 	TShaderMapRef<TOneColorVS<true> > VertexShader(ShaderMap);
 	static FGlobalBoundShaderState BoundShaderState;
 	SetGlobalBoundShaderState(RHICmdList, FeatureLevel, BoundShaderState, GetVertexDeclarationFVector4(), *VertexShader, nullptr);
-	GEngine->HMDDevice->DrawHiddenAreaMaskView_RenderThread(RHICmdList, View.StereoPass);
+	GEngine->HMDDevice->DrawHiddenAreaMesh_RenderThread(RHICmdList, View.StereoPass);
 }
 
 bool FDeferredShadingSceneRenderer::RenderPrePassView(FRHICommandList& RHICmdList, const FViewInfo& View)
diff --git a/Engine/Source/Runtime/Renderer/Private/PostProcess/RenderingCompositionGraph.cpp b/Engine/Source/Runtime/Renderer/Private/PostProcess/RenderingCompositionGraph.cpp
--- a/Engine/Source/Runtime/Renderer/Private/PostProcess/RenderingCompositionGraph.cpp
+++ b/Engine/Source/Runtime/Renderer/Private/PostProcess/RenderingCompositionGraph.cpp
@@ -147,7 +147,7 @@ void FRenderingCompositePassContext::Process(FRenderingCompositePass* Root, cons
 		HiddenAreaMaskCVar->GetValueOnRenderThread() == 1 &&
 		GEngine &&
 		GEngine->HMDDevice.IsValid() &&
-		GEngine->HMDDevice->HasPostProcessMesh());
+		GEngine->HMDDevice->HasVisibleAreaMesh());
 
 	if(Root)
 	{
diff --git a/Engine/Source/Runtime/Renderer/Private/PostProcess/SceneFilterRendering.cpp b/Engine/Source/Runtime/Renderer/Private/PostProcess/SceneFilterRendering.cpp
--- a/Engine/Source/Runtime/Renderer/Private/PostProcess/SceneFilterRendering.cpp
+++ b/Engine/Source/Runtime/Renderer/Private/PostProcess/SceneFilterRendering.cpp
@@ -328,7 +328,7 @@ void DrawHmdMesh(
 
 	SetUniformBufferParameterImmediate(RHICmdList, VertexShader->GetVertexShader(), VertexShader->GetUniformBufferParameter<FDrawRectangleParameters>(), Parameters);
 
-	GEngine->HMDDevice->DrawPostProcessMesh_RenderThread(RHICmdList, StereoView);
+	GEngine->HMDDevice->DrawVisibleAreaMesh_RenderThread(RHICmdList, StereoView);
 }
 
 void DrawPostProcessPass(
