Added RebaseObjectOrientationAndPosition to IHeadMountedDisplay. This allows for an external object to be remapped to the HMD's current base.

diff --git a/Engine/Plugins/Experimental/SimpleHMD/Source/SimpleHMD/Private/SimpleHMD.cpp b/Engine/Plugins/Experimental/SimpleHMD/Source/SimpleHMD/Private/SimpleHMD.cpp
--- a/Engine/Plugins/Experimental/SimpleHMD/Source/SimpleHMD/Private/SimpleHMD.cpp
+++ b/Engine/Plugins/Experimental/SimpleHMD/Source/SimpleHMD/Private/SimpleHMD.cpp
@@ -80,6 +80,10 @@ void FSimpleHMD::GetPositionalTrackingCameraProperties(FVector& OutOrigin, FQuat
 {
 }
 
+void FSimpleHMD::RebaseObjectOrientationAndPosition(FVector& OutPosition, FQuat& OutOrientation) const
+{
+}
+
 void FSimpleHMD::SetInterpupillaryDistance(float NewInterpupillaryDistance)
 {
 }
diff --git a/Engine/Plugins/Experimental/SimpleHMD/Source/SimpleHMD/Private/SimpleHMD.h b/Engine/Plugins/Experimental/SimpleHMD/Source/SimpleHMD/Private/SimpleHMD.h
--- a/Engine/Plugins/Experimental/SimpleHMD/Source/SimpleHMD/Private/SimpleHMD.h
+++ b/Engine/Plugins/Experimental/SimpleHMD/Source/SimpleHMD/Private/SimpleHMD.h
@@ -25,6 +25,7 @@ public:
 	virtual bool DoesSupportPositionalTracking() const override;
 	virtual bool HasValidTrackingPosition() override;
 	virtual void GetPositionalTrackingCameraProperties(FVector& OutOrigin, FQuat& OutOrientation, float& OutHFOV, float& OutVFOV, float& OutCameraDistance, float& OutNearPlane, float& OutFarPlane) const override;
+	virtual void RebaseObjectOrientationAndPosition(FVector& OutPosition, FQuat& OutOrientation) const override;
 
 	virtual void SetInterpupillaryDistance(float NewInterpupillaryDistance) override;
 	virtual float GetInterpupillaryDistance() const override;
diff --git a/Engine/Plugins/Runtime/OculusRift/Source/OculusRift/Private/OculusRiftHMD.cpp b/Engine/Plugins/Runtime/OculusRift/Source/OculusRift/Private/OculusRiftHMD.cpp
--- a/Engine/Plugins/Runtime/OculusRift/Source/OculusRift/Private/OculusRiftHMD.cpp
+++ b/Engine/Plugins/Runtime/OculusRift/Source/OculusRift/Private/OculusRiftHMD.cpp
@@ -271,6 +271,10 @@ void FOculusRiftHMD::GetPositionalTrackingCameraProperties(FVector& OutOrigin, F
 	OutOrigin = Pos + frame->Settings->PositionOffset;
 }
 
+void FOculusRiftHMD::RebaseObjectOrientationAndPosition(FVector& OutPosition, FQuat& OutOrientation) const
+{
+}
+
 bool FOculusRiftHMD::IsInLowPersistenceMode() const
 {
 	const auto frame = GetFrame();
diff --git a/Engine/Plugins/Runtime/OculusRift/Source/OculusRift/Private/OculusRiftHMD.h b/Engine/Plugins/Runtime/OculusRift/Source/OculusRift/Private/OculusRiftHMD.h
--- a/Engine/Plugins/Runtime/OculusRift/Source/OculusRift/Private/OculusRiftHMD.h
+++ b/Engine/Plugins/Runtime/OculusRift/Source/OculusRift/Private/OculusRiftHMD.h
@@ -279,6 +279,7 @@ public:
 	virtual bool DoesSupportPositionalTracking() const override;
 	virtual bool HasValidTrackingPosition() override;
 	virtual void GetPositionalTrackingCameraProperties(FVector& OutOrigin, FQuat& OutOrientation, float& OutHFOV, float& OutVFOV, float& OutCameraDistance, float& OutNearPlane, float& OutFarPlane) const override;
+	virtual void RebaseObjectOrientationAndPosition(FVector& OutPosition, FQuat& OutOrientation) const override;
 
 	virtual TSharedPtr<class ISceneViewExtension, ESPMode::ThreadSafe> GetViewExtension() override;
 	virtual bool Exec( UWorld* InWorld, const TCHAR* Cmd, FOutputDevice& Ar ) override;
diff --git a/Engine/Plugins/Runtime/Steam/SteamVR/Source/SteamVR/Private/SteamVRHMD.cpp b/Engine/Plugins/Runtime/Steam/SteamVR/Source/SteamVR/Private/SteamVRHMD.cpp
--- a/Engine/Plugins/Runtime/Steam/SteamVR/Source/SteamVR/Private/SteamVRHMD.cpp
+++ b/Engine/Plugins/Runtime/Steam/SteamVR/Source/SteamVR/Private/SteamVRHMD.cpp
@@ -169,6 +169,10 @@ void FSteamVRHMD::GetPositionalTrackingCameraProperties(FVector& OutOrigin, FQua
 {
 }
 
+void FSteamVRHMD::RebaseObjectOrientationAndPosition(FVector& OutPosition, FQuat& OutOrientation) const
+{
+}
+
 void FSteamVRHMD::SetInterpupillaryDistance(float NewInterpupillaryDistance)
 {
 }
diff --git a/Engine/Plugins/Runtime/Steam/SteamVR/Source/SteamVR/Private/SteamVRHMD.h b/Engine/Plugins/Runtime/Steam/SteamVR/Source/SteamVR/Private/SteamVRHMD.h
--- a/Engine/Plugins/Runtime/Steam/SteamVR/Source/SteamVR/Private/SteamVRHMD.h
+++ b/Engine/Plugins/Runtime/Steam/SteamVR/Source/SteamVR/Private/SteamVRHMD.h
@@ -48,6 +48,7 @@ public:
 	virtual bool DoesSupportPositionalTracking() const override;
 	virtual bool HasValidTrackingPosition() override;
 	virtual void GetPositionalTrackingCameraProperties(FVector& OutOrigin, FQuat& OutOrientation, float& OutHFOV, float& OutVFOV, float& OutCameraDistance, float& OutNearPlane, float& OutFarPlane) const override;
+	virtual void RebaseObjectOrientationAndPosition(FVector& OutPosition, FQuat& OutOrientation) const override;
 
 	virtual void SetInterpupillaryDistance(float NewInterpupillaryDistance) override;
 	virtual float GetInterpupillaryDistance() const override;
diff --git a/Engine/Source/Runtime/HeadMountedDisplay/Public/IHeadMountedDisplay.h b/Engine/Source/Runtime/HeadMountedDisplay/Public/IHeadMountedDisplay.h
--- a/Engine/Source/Runtime/HeadMountedDisplay/Public/IHeadMountedDisplay.h
+++ b/Engine/Source/Runtime/HeadMountedDisplay/Public/IHeadMountedDisplay.h
@@ -95,6 +95,11 @@ public:
 	 */
 	virtual void GetCurrentOrientationAndPosition(FQuat& CurrentOrientation, FVector& CurrentPosition) = 0;
 
+	/**
+	 * Rebase the input position and orientation to that of the HMD's base
+	 */
+	virtual void RebaseObjectOrientationAndPosition(FVector& Position, FQuat& Orientation) const = 0;
+
 	/**
 	 * Get the ISceneViewExtension for this HMD, or none.
 	 */
