#oculus Fix for XP build breaks, wrapped functions with OCULUS_RIFT_SUPPORTED_PLATFORMS

diff --git a/Engine/Plugins/Runtime/OculusRift/Source/OculusRift/Private/OculusRiftHMD.cpp b/Engine/Plugins/Runtime/OculusRift/Source/OculusRift/Private/OculusRiftHMD.cpp
--- a/Engine/Plugins/Runtime/OculusRift/Source/OculusRift/Private/OculusRiftHMD.cpp
+++ b/Engine/Plugins/Runtime/OculusRift/Source/OculusRift/Private/OculusRiftHMD.cpp
@@ -34,6 +34,7 @@ class FOculusRiftPlugin : public IOculusRiftPlugin
 		check(IsInGameThread());
 		bool bRetVal = false;
 
+#if OCULUS_RIFT_SUPPORTED_PLATFORMS
 		IHeadMountedDisplay* HMD = HeadMountedDisplay.Pin().Get();
 		if (HMD && HMD->GetHMDDeviceType() == EHMDDeviceType::DT_OculusRift)
 		{
@@ -46,6 +47,7 @@ class FOculusRiftPlugin : public IOculusRiftPlugin
 				bRetVal = true;
 			}
 		}
+#endif //OCULUS_RIFT_SUPPORTED_PLATFORMS
 		return bRetVal;
 	}
 
@@ -54,12 +56,14 @@ class FOculusRiftPlugin : public IOculusRiftPlugin
 		check(IsInGameThread());
 		ovrHmd bRetVal = nullptr;
 
+#if OCULUS_RIFT_SUPPORTED_PLATFORMS
 		IHeadMountedDisplay* HMD = HeadMountedDisplay.Pin().Get();
 		if (HMD && HMD->GetHMDDeviceType() == EHMDDeviceType::DT_OculusRift)
 		{
 			FOculusRiftHMD* OculusHMD = static_cast<FOculusRiftHMD*>(HMD);
 			bRetVal = OculusHMD->Hmd;
 		}
+#endif //OCULUS_RIFT_SUPPORTED_PLATFORMS
 		return bRetVal;
 	}
 
@@ -68,6 +72,7 @@ class FOculusRiftPlugin : public IOculusRiftPlugin
 		check(IsInGameThread());
 		bool bRetVal = false;
 
+#if OCULUS_RIFT_SUPPORTED_PLATFORMS
 		IHeadMountedDisplay* HMD = HeadMountedDisplay.Pin().Get();
 		if (TrackingState && HMD && HMD->GetHMDDeviceType() == EHMDDeviceType::DT_OculusRift)
 		{
@@ -79,6 +84,7 @@ class FOculusRiftPlugin : public IOculusRiftPlugin
 				bRetVal = true;
 			}
 		}
+#endif //OCULUS_RIFT_SUPPORTED_PLATFORMS
 		return bRetVal;
 	}
 
