#oculus Integrating fix for texture resolution reporting

diff --git a/Engine/Plugins/Runtime/OculusRift/Source/OculusRift/Private/OculusRiftRenderD3D.cpp b/Engine/Plugins/Runtime/OculusRift/Source/OculusRift/Private/OculusRiftRenderD3D.cpp
--- a/Engine/Plugins/Runtime/OculusRift/Source/OculusRift/Private/OculusRiftRenderD3D.cpp
+++ b/Engine/Plugins/Runtime/OculusRift/Source/OculusRift/Private/OculusRiftRenderD3D.cpp
@@ -452,8 +452,8 @@ void FOculusRiftHMD::D3D11Bridge::BeginRendering(FHMDViewExtension& InRenderCont
 		ActualMirrorWindowSize.X != 0 && ActualMirrorWindowSize.Y != 0)
 	{
 		D3D11_TEXTURE2D_DESC dsDesc;
-		dsDesc.Width = ActualMirrorWindowSize.X;
-		dsDesc.Height = ActualMirrorWindowSize.Y;
+		dsDesc.Width = (UINT)ActualMirrorWindowSize.X;
+		dsDesc.Height = (UINT)ActualMirrorWindowSize.Y;
 		dsDesc.MipLevels = 1;
 		dsDesc.ArraySize = 1;
 		dsDesc.Format = DXGI_FORMAT_B8G8R8A8_UNORM_SRGB; // SRGB is required for the compositor
@@ -473,7 +473,7 @@ void FOculusRiftHMD::D3D11Bridge::BeginRendering(FHMDViewExtension& InRenderCont
 			return;
 		}
 
-		UE_LOG(LogHMD, Log, TEXT("Allocated a new mirror texture (size %d x %d)"), ActualMirrorWindowSize.X, ActualMirrorWindowSize.Y);
+		UE_LOG(LogHMD, Log, TEXT("Allocated a new mirror texture (size %d x %d)"), (int)ActualMirrorWindowSize.X, (int)ActualMirrorWindowSize.Y);
 		ovrD3D11Texture D3DMirrorTexture;
 		D3DMirrorTexture.Texture = *MirrorTexture;
 		MirrorTextureRHI = D3D11CreateTexture2DAlias(
