Moving the hidden area mesh rendering outside of the pre-pass. It runs before the pre-pass if enabled or before the base pass if pre-pass is disabled.
Adding a full render target clear to the tonemap post process when rendering in stereo with the visible area mesh active. This ensures we clear the area outside of the visible area mesh each frame.
Adding hidden/visible area mesh support to the current rift EVT models.

diff --git a/Engine/Plugins/Runtime/OculusRift/Source/OculusRift/Private/OculusRiftHMD.cpp b/Engine/Plugins/Runtime/OculusRift/Source/OculusRift/Private/OculusRiftHMD.cpp
--- a/Engine/Plugins/Runtime/OculusRift/Source/OculusRift/Private/OculusRiftHMD.cpp
+++ b/Engine/Plugins/Runtime/OculusRift/Source/OculusRift/Private/OculusRiftHMD.cpp
@@ -1412,6 +1412,13 @@ void FOculusRiftHMD::SetupOcclusionMeshes()
 		VisibleAreaMeshes[0].BuildMesh(CB_LeftEyeVisibleAreaPositions, VisibleAreaVertexCount, FHMDViewMesh::MT_VisibleArea);
 		VisibleAreaMeshes[1].BuildMesh(CB_RightEyeVisibleAreaPositions, VisibleAreaVertexCount, FHMDViewMesh::MT_VisibleArea);
 	}
+	else if (HmdDesc.Type == ovrHmdType::ovrHmd_E3_2015 || HmdDesc.Type == ovrHmdType::ovrHmd_ES06)
+	{
+		HiddenAreaMeshes[0].BuildMesh(EVT_LeftEyeHiddenAreaPositions, HiddenAreaVertexCount, FHMDViewMesh::MT_HiddenArea);
+		HiddenAreaMeshes[1].BuildMesh(EVT_RightEyeHiddenAreaPositions, HiddenAreaVertexCount, FHMDViewMesh::MT_HiddenArea);
+		VisibleAreaMeshes[0].BuildMesh(EVT_LeftEyeVisibleAreaPositions, VisibleAreaVertexCount, FHMDViewMesh::MT_VisibleArea);
+		VisibleAreaMeshes[1].BuildMesh(EVT_RightEyeVisibleAreaPositions, VisibleAreaVertexCount, FHMDViewMesh::MT_VisibleArea);
+	}
 }
 
 void FOculusRiftHMD::UpdateHmdCaps()
diff --git a/Engine/Plugins/Runtime/OculusRift/Source/OculusRift/Private/OculusRiftMeshAssets.h b/Engine/Plugins/Runtime/OculusRift/Source/OculusRift/Private/OculusRiftMeshAssets.h
--- a/Engine/Plugins/Runtime/OculusRift/Source/OculusRift/Private/OculusRiftMeshAssets.h
+++ b/Engine/Plugins/Runtime/OculusRift/Source/OculusRift/Private/OculusRiftMeshAssets.h
@@ -5,6 +5,206 @@
 static const uint32 HiddenAreaVertexCount = 96;
 static const uint32 VisibleAreaVertexCount = 102;
 
+static const FVector2D EVT_LeftEyeHiddenAreaPositions[] =
+{
+	FVector2D(0.6420568228f, 1.0f),
+	FVector2D(1.0f, 1.0f),
+	FVector2D(0.7550624013f, 0.9892609119f),
+	FVector2D(0.7550624013f, 0.9892609119f),
+	FVector2D(1.0f, 1.0f),
+	FVector2D(0.8259530067f, 0.9766981006f),
+	FVector2D(1.0f, 0.6303145289f),
+	FVector2D(0.9933852553f, 0.6894935966f),
+	FVector2D(1.0f, 1.0f),
+	FVector2D(0.9780833721f, 0.7837033868f),
+	FVector2D(1.0f, 1.0f),
+	FVector2D(0.9933852553f, 0.6894935966f),
+	FVector2D(0.8259530067f, 0.9766981006f),
+	FVector2D(1.0f, 1.0f),
+	FVector2D(0.8814527988f, 0.9645425081f),
+	FVector2D(0.9561929107f, 0.8690519333f),
+	FVector2D(1.0f, 1.0f),
+	FVector2D(0.9780833721f, 0.7837033868f),
+	FVector2D(0.8814527988f, 0.9645425081f),
+	FVector2D(1.0f, 1.0f),
+	FVector2D(0.9246563315f, 0.9536676407f),
+	FVector2D(0.9246563315f, 0.9536676407f),
+	FVector2D(1.0f, 1.0f),
+	FVector2D(0.9561929107f, 0.8690519333f),
+	FVector2D(1.0f, 0.4252603054f),
+	FVector2D(1.0f, 0.0f),
+	FVector2D(0.9865142703f, 0.3167535067f),
+	FVector2D(0.6295351982f, 0.0f),
+	FVector2D(0.6981808543f, 0.005146139301f),
+	FVector2D(1.0f, 0.0f),
+	FVector2D(0.7825559974f, 0.01721908338f),
+	FVector2D(1.0f, 0.0f),
+	FVector2D(0.6981808543f, 0.005146139301f),
+	FVector2D(0.9865142703f, 0.3167535067f),
+	FVector2D(1.0f, 0.0f),
+	FVector2D(0.9636768699f, 0.216376245f),
+	FVector2D(0.8487468362f, 0.03234770149f),
+	FVector2D(1.0f, 0.0f),
+	FVector2D(0.7825559974f, 0.01721908338f),
+	FVector2D(0.9636768699f, 0.216376245f),
+	FVector2D(1.0f, 0.0f),
+	FVector2D(0.9353111982f, 0.131273821f),
+	FVector2D(0.8951631188f, 0.04515158013f),
+	FVector2D(1.0f, 0.0f),
+	FVector2D(0.8487468362f, 0.03234770149f),
+	FVector2D(0.9353111982f, 0.131273821f),
+	FVector2D(1.0f, 0.0f),
+	FVector2D(0.8951631188f, 0.04515158013f),
+	FVector2D(0.0f, 0.0f),
+	FVector2D(0.4338365495f, 0.0f),
+	FVector2D(0.330127269f, 0.01285663527f),
+	FVector2D(0.0f, 0.4224827588f),
+	FVector2D(0.01270084083f, 0.3334831297f),
+	FVector2D(0.0f, 0.0f),
+	FVector2D(0.330127269f, 0.01285663527f),
+	FVector2D(0.0f, 0.0f),
+	FVector2D(0.2493925691f, 0.0313315168f),
+	FVector2D(0.03772026673f, 0.2338332087f),
+	FVector2D(0.0f, 0.0f),
+	FVector2D(0.01270084083f, 0.3334831297f),
+	FVector2D(0.2493925691f, 0.0313315168f),
+	FVector2D(0.0f, 0.0f),
+	FVector2D(0.1802922487f, 0.05096876994f),
+	FVector2D(0.07248612493f, 0.1436391324f),
+	FVector2D(0.0f, 0.0f),
+	FVector2D(0.03772026673f, 0.2338332087f),
+	FVector2D(0.1802922487f, 0.05096876994f),
+	FVector2D(0.0f, 0.0f),
+	FVector2D(0.1079794019f, 0.07832186669f),
+	FVector2D(0.1079794019f, 0.07832186669f),
+	FVector2D(0.0f, 0.0f),
+	FVector2D(0.07248612493f, 0.1436391324f),
+	FVector2D(0.4444268048f, 1.0f),
+	FVector2D(0.3694992661f, 0.9932333231f),
+	FVector2D(0.0f, 1.0f),
+	FVector2D(0.0f, 0.6263416409f),
+	FVector2D(0.0f, 1.0f),
+	FVector2D(0.017040601f, 0.7411668301f),
+	FVector2D(0.268977195f, 0.9784004092f),
+	FVector2D(0.3694992661f, 0.9932333231f),
+	FVector2D(0.0f, 1.0f),
+	FVector2D(0.017040601f, 0.7411668301f),
+	FVector2D(0.03240321949f, 0.8031471968f),
+	FVector2D(0.0f, 1.0f),
+	FVector2D(0.1714757681f, 0.9543109536f),
+	FVector2D(0.268977195f, 0.9784004092f),
+	FVector2D(0.0f, 1.0f),
+	FVector2D(0.03240321949f, 0.8031471968f),
+	FVector2D(0.05181787908f, 0.8630086184f),
+	FVector2D(0.0f, 1.0f),
+	FVector2D(0.07756151259f, 0.9210018516f),
+	FVector2D(0.1714757681f, 0.9543109536f),
+	FVector2D(0.0f, 1.0f),
+	FVector2D(0.05181787908f, 0.8630086184f),
+	FVector2D(0.07756151259f, 0.9210018516f),
+	FVector2D(0.0f, 1.0f),
+};
+
+static const FVector2D EVT_RightEyeHiddenAreaPositions[] =
+{
+	FVector2D(0.3579431772f, 1.0f),
+	FVector2D(0.2449375987f, 0.9892609119f),
+	FVector2D(0.0f, 1.0f),
+	FVector2D(0.0f, 1.0f),
+	FVector2D(0.0f, 0.6303145289f),
+	FVector2D(0.006614744663f, 0.6894935966f),
+	FVector2D(0.006614744663f, 0.6894935966f),
+	FVector2D(0.0f, 1.0f),
+	FVector2D(0.02191662788f, 0.7837033868f),
+	FVector2D(0.1740469933f, 0.9766981006f),
+	FVector2D(0.0f, 1.0f),
+	FVector2D(0.2449375987f, 0.9892609119f),
+	FVector2D(0.02191662788f, 0.7837033868f),
+	FVector2D(0.0f, 1.0f),
+	FVector2D(0.04380708933f, 0.8690519333f),
+	FVector2D(0.1185472012f, 0.9645425081f),
+	FVector2D(0.0f, 1.0f),
+	FVector2D(0.1740469933f, 0.9766981006f),
+	FVector2D(0.07534366846f, 0.9536676407f),
+	FVector2D(0.0f, 1.0f),
+	FVector2D(0.1185472012f, 0.9645425081f),
+	FVector2D(0.04380708933f, 0.8690519333f),
+	FVector2D(0.0f, 1.0f),
+	FVector2D(0.07534366846f, 0.9536676407f),
+	FVector2D(0.0f, 0.4252603054f),
+	FVector2D(0.01348572969f, 0.3167535067f),
+	FVector2D(0.0f, 0.0f),
+	FVector2D(0.0f, 0.0f),
+	FVector2D(0.3704648018f, 0.0f),
+	FVector2D(0.3018191457f, 0.005146139301f),
+	FVector2D(0.3018191457f, 0.005146139301f),
+	FVector2D(0.0f, 0.0f),
+	FVector2D(0.2174440026f, 0.01721908338f),
+	FVector2D(0.2174440026f, 0.01721908338f),
+	FVector2D(0.0f, 0.0f),
+	FVector2D(0.1512531638f, 0.03234770149f),
+	FVector2D(0.03632313013f, 0.216376245f),
+	FVector2D(0.0f, 0.0f),
+	FVector2D(0.01348572969f, 0.3167535067f),
+	FVector2D(0.1512531638f, 0.03234770149f),
+	FVector2D(0.0f, 0.0f),
+	FVector2D(0.1048368812f, 0.04515158013f),
+	FVector2D(0.06468880177f, 0.131273821f),
+	FVector2D(0.0f, 0.0f),
+	FVector2D(0.03632313013f, 0.216376245f),
+	FVector2D(0.1048368812f, 0.04515158013f),
+	FVector2D(0.0f, 0.0f),
+	FVector2D(0.06468880177f, 0.131273821f),
+	FVector2D(1.0f, 0.0f),
+	FVector2D(1.0f, 0.4224827588f),
+	FVector2D(0.9872991443f, 0.3334831297f),
+	FVector2D(0.6698727608f, 0.01285663527f),
+	FVector2D(0.5661634207f, 0.0f),
+	FVector2D(1.0f, 0.0f),
+	FVector2D(0.9872991443f, 0.3334831297f),
+	FVector2D(1.0f, 0.0f),
+	FVector2D(0.962279737f, 0.2338332087f),
+	FVector2D(0.7506074309f, 0.0313315168f),
+	FVector2D(1.0f, 0.0f),
+	FVector2D(0.6698727608f, 0.01285663527f),
+	FVector2D(0.8197077513f, 0.05096876994f),
+	FVector2D(1.0f, 0.0f),
+	FVector2D(0.7506074309f, 0.0313315168f),
+	FVector2D(0.962279737f, 0.2338332087f),
+	FVector2D(1.0f, 0.0f),
+	FVector2D(0.9275138378f, 0.1436391324f),
+	FVector2D(0.8920205832f, 0.07832186669f),
+	FVector2D(1.0f, 0.0f),
+	FVector2D(0.8197077513f, 0.05096876994f),
+	FVector2D(0.9275138378f, 0.1436391324f),
+	FVector2D(1.0f, 0.0f),
+	FVector2D(0.8920205832f, 0.07832186669f),
+	FVector2D(0.555573225f, 1.0f),
+	FVector2D(1.0f, 1.0f),
+	FVector2D(0.6305007339f, 0.9932333231f),
+	FVector2D(0.6305007339f, 0.9932333231f),
+	FVector2D(0.7310228348f, 0.9784004092f),
+	FVector2D(1.0f, 1.0f),
+	FVector2D(0.9829593897f, 0.7411668301f),
+	FVector2D(1.0f, 0.6263416409f),
+	FVector2D(1.0f, 1.0f),
+	FVector2D(0.7310228348f, 0.9784004092f),
+	FVector2D(0.8285242319f, 0.9543109536f),
+	FVector2D(1.0f, 1.0f),
+	FVector2D(0.9675967693f, 0.8031471968f),
+	FVector2D(0.9829593897f, 0.7411668301f),
+	FVector2D(1.0f, 1.0f),
+	FVector2D(0.948182106f, 0.8630086184f),
+	FVector2D(0.9675967693f, 0.8031471968f),
+	FVector2D(1.0f, 1.0f),
+	FVector2D(0.8285242319f, 0.9543109536f),
+	FVector2D(0.9224384427f, 0.9210018516f),
+	FVector2D(1.0f, 1.0f),
+	FVector2D(0.9224384427f, 0.9210018516f),
+	FVector2D(0.948182106f, 0.8630086184f),
+	FVector2D(1.0f, 1.0f),
+};
+
 static const FVector2D DK2_LeftEyeHiddenAreaPositions[] =
 {
 	FVector2D(0.0f, 1.0f),
@@ -405,6 +605,218 @@ static const FVector2D CB_RightEyeHiddenAreaPositions[] =
 	FVector2D(0.0734243989f, 0.131273821f),
 };
 
+static const FVector2D EVT_LeftEyeVisibleAreaPositions[] =
+{
+	FVector2D(0.9246563315f, 0.9536676407f),
+	FVector2D(0.9561929107f, 0.8690519333f),
+	FVector2D(0.8814527988f, 0.9645425081f),
+	FVector2D(0.8951631188f, 0.04515158013f),
+	FVector2D(0.8487468362f, 0.03234770149f),
+	FVector2D(0.9353111982f, 0.131273821f),
+	FVector2D(0.07756151259f, 0.9210018516f),
+	FVector2D(0.1714757681f, 0.9543109536f),
+	FVector2D(0.05181787908f, 0.8630086184f),
+	FVector2D(0.1079794019f, 0.07832186669f),
+	FVector2D(0.07248612493f, 0.1436391324f),
+	FVector2D(0.1802922487f, 0.05096876994f),
+	FVector2D(0.03240321949f, 0.8031471968f),
+	FVector2D(0.05181787908f, 0.8630086184f),
+	FVector2D(0.017040601f, 0.7411668301f),
+	FVector2D(0.0f, 0.6263416409f),
+	FVector2D(0.017040601f, 0.7411668301f),
+	FVector2D(0.0f, 0.4224827588f),
+	FVector2D(0.01270084083f, 0.3334831297f),
+	FVector2D(0.0f, 0.4224827588f),
+	FVector2D(0.03772026673f, 0.2338332087f),
+	FVector2D(0.2493925691f, 0.0313315168f),
+	FVector2D(0.1802922487f, 0.05096876994f),
+	FVector2D(0.330127269f, 0.01285663527f),
+	FVector2D(0.4338365495f, 0.0f),
+	FVector2D(0.330127269f, 0.01285663527f),
+	FVector2D(0.6295351982f, 0.0f),
+	FVector2D(0.6981808543f, 0.005146139301f),
+	FVector2D(0.6295351982f, 0.0f),
+	FVector2D(0.7825559974f, 0.01721908338f),
+	FVector2D(0.9636768699f, 0.216376245f),
+	FVector2D(0.9353111982f, 0.131273821f),
+	FVector2D(0.9865142703f, 0.3167535067f),
+	FVector2D(1.0f, 0.4252603054f),
+	FVector2D(0.9865142703f, 0.3167535067f),
+	FVector2D(1.0f, 0.6303145289f),
+	FVector2D(0.9933852553f, 0.6894935966f),
+	FVector2D(1.0f, 0.6303145289f),
+	FVector2D(0.9780833721f, 0.7837033868f),
+	FVector2D(0.8259530067f, 0.9766981006f),
+	FVector2D(0.8814527988f, 0.9645425081f),
+	FVector2D(0.7550624013f, 0.9892609119f),
+	FVector2D(0.6420568228f, 1.0f),
+	FVector2D(0.7550624013f, 0.9892609119f),
+	FVector2D(0.4444268048f, 1.0f),
+	FVector2D(0.3694992661f, 0.9932333231f),
+	FVector2D(0.4444268048f, 1.0f),
+	FVector2D(0.268977195f, 0.9784004092f),
+	FVector2D(0.7550624013f, 0.9892609119f),
+	FVector2D(0.8814527988f, 0.9645425081f),
+	FVector2D(0.9561929107f, 0.8690519333f),
+	FVector2D(0.9780833721f, 0.7837033868f),
+	FVector2D(0.7550624013f, 0.9892609119f),
+	FVector2D(0.9561929107f, 0.8690519333f),
+	FVector2D(0.7550624013f, 0.9892609119f),
+	FVector2D(0.9780833721f, 0.7837033868f),
+	FVector2D(1.0f, 0.6303145289f),
+	FVector2D(0.4444268048f, 1.0f),
+	FVector2D(0.7550624013f, 0.9892609119f),
+	FVector2D(1.0f, 0.6303145289f),
+	FVector2D(0.4444268048f, 1.0f),
+	FVector2D(1.0f, 0.6303145289f),
+	FVector2D(0.9865142703f, 0.3167535067f),
+	FVector2D(0.268977195f, 0.9784004092f),
+	FVector2D(0.4444268048f, 1.0f),
+	FVector2D(0.9865142703f, 0.3167535067f),
+	FVector2D(0.268977195f, 0.9784004092f),
+	FVector2D(0.9865142703f, 0.3167535067f),
+	FVector2D(0.9353111982f, 0.131273821f),
+	FVector2D(0.268977195f, 0.9784004092f),
+	FVector2D(0.9353111982f, 0.131273821f),
+	FVector2D(0.8487468362f, 0.03234770149f),
+	FVector2D(0.7825559974f, 0.01721908338f),
+	FVector2D(0.268977195f, 0.9784004092f),
+	FVector2D(0.8487468362f, 0.03234770149f),
+	FVector2D(0.268977195f, 0.9784004092f),
+	FVector2D(0.7825559974f, 0.01721908338f),
+	FVector2D(0.6295351982f, 0.0f),
+	FVector2D(0.268977195f, 0.9784004092f),
+	FVector2D(0.6295351982f, 0.0f),
+	FVector2D(0.330127269f, 0.01285663527f),
+	FVector2D(0.268977195f, 0.9784004092f),
+	FVector2D(0.330127269f, 0.01285663527f),
+	FVector2D(0.1802922487f, 0.05096876994f),
+	FVector2D(0.268977195f, 0.9784004092f),
+	FVector2D(0.1802922487f, 0.05096876994f),
+	FVector2D(0.07248612493f, 0.1436391324f),
+	FVector2D(0.03772026673f, 0.2338332087f),
+	FVector2D(0.268977195f, 0.9784004092f),
+	FVector2D(0.07248612493f, 0.1436391324f),
+	FVector2D(0.268977195f, 0.9784004092f),
+	FVector2D(0.03772026673f, 0.2338332087f),
+	FVector2D(0.0f, 0.4224827588f),
+	FVector2D(0.268977195f, 0.9784004092f),
+	FVector2D(0.0f, 0.4224827588f),
+	FVector2D(0.017040601f, 0.7411668301f),
+	FVector2D(0.268977195f, 0.9784004092f),
+	FVector2D(0.1714757681f, 0.9543109536f),
+	FVector2D(0.017040601f, 0.7411668301f),
+	FVector2D(0.1714757681f, 0.9543109536f),
+	FVector2D(0.017040601f, 0.7411668301f),
+	FVector2D(0.05181787908f, 0.8630086184f),
+};
+
+static const FVector2D EVT_RightEyeVisibleAreaPositions[] =
+{
+	FVector2D(0.1185472012f, 0.9645425081f),
+	FVector2D(0.07534366846f, 0.9536676407f),
+	FVector2D(0.04380708933f, 0.8690519333f),
+	FVector2D(0.06468880177f, 0.131273821f),
+	FVector2D(0.1048368812f, 0.04515158013f),
+	FVector2D(0.1512531638f, 0.03234770149f),
+	FVector2D(0.9224385023f, 0.9210018516f),
+	FVector2D(0.948182106f, 0.8630086184f),
+	FVector2D(0.8285242319f, 0.9543109536f),
+	FVector2D(0.8197077513f, 0.05096876994f),
+	FVector2D(0.8920205832f, 0.07832186669f),
+	FVector2D(0.9275138378f, 0.1436391324f),
+	FVector2D(0.8285242319f, 0.9543109536f),
+	FVector2D(0.7310228348f, 0.9784004092f),
+	FVector2D(0.6305007339f, 0.9932333231f),
+	FVector2D(0.6305007339f, 0.9932333231f),
+	FVector2D(0.555573225f, 1.0f),
+	FVector2D(0.3579431772f, 1.0f),
+	FVector2D(0.3579431772f, 1.0f),
+	FVector2D(0.2449375987f, 0.9892609119f),
+	FVector2D(0.1740469933f, 0.9766981006f),
+	FVector2D(0.04380708933f, 0.8690519333f),
+	FVector2D(0.02191662788f, 0.7837033868f),
+	FVector2D(0.006614744663f, 0.6894935966f),
+	FVector2D(0.006614744663f, 0.6894935966f),
+	FVector2D(0.0f, 0.6303145289f),
+	FVector2D(0.0f, 0.4252603054f),
+	FVector2D(0.0f, 0.4252603054f),
+	FVector2D(0.01348572969f, 0.3167535067f),
+	FVector2D(0.03632313013f, 0.216376245f),
+	FVector2D(0.1512531638f, 0.03234770149f),
+	FVector2D(0.2174440026f, 0.01721908338f),
+	FVector2D(0.3018191457f, 0.005146139301f),
+	FVector2D(0.3018191457f, 0.005146139301f),
+	FVector2D(0.3704648018f, 0.0f),
+	FVector2D(0.5661634207f, 0.0f),
+	FVector2D(0.5661634207f, 0.0f),
+	FVector2D(0.6698727608f, 0.01285663527f),
+	FVector2D(0.7506074309f, 0.0313315168f),
+	FVector2D(0.9275138378f, 0.1436391324f),
+	FVector2D(0.962279737f, 0.2338332087f),
+	FVector2D(0.9872991443f, 0.3334831297f),
+	FVector2D(0.9872991443f, 0.3334831297f),
+	FVector2D(1.0f, 0.4224827588f),
+	FVector2D(1.0f, 0.6263416409f),
+	FVector2D(1.0f, 0.6263416409f),
+	FVector2D(0.9829593897f, 0.7411668301f),
+	FVector2D(0.9675967693f, 0.8031471968f),
+	FVector2D(0.1740469933f, 0.9766981006f),
+	FVector2D(0.1185472012f, 0.9645425081f),
+	FVector2D(0.04380708933f, 0.8690519333f),
+	FVector2D(0.1740469933f, 0.9766981006f),
+	FVector2D(0.04380708933f, 0.8690519333f),
+	FVector2D(0.006614744663f, 0.6894935966f),
+	FVector2D(0.3579431772f, 1.0f),
+	FVector2D(0.1740469933f, 0.9766981006f),
+	FVector2D(0.006614744663f, 0.6894935966f),
+	FVector2D(0.3579431772f, 1.0f),
+	FVector2D(0.006614744663f, 0.6894935966f),
+	FVector2D(0.0f, 0.4252603054f),
+	FVector2D(0.6305007339f, 0.9932333231f),
+	FVector2D(0.3579431772f, 1.0f),
+	FVector2D(0.0f, 0.4252603054f),
+	FVector2D(0.6305007339f, 0.9932333231f),
+	FVector2D(0.0f, 0.4252603054f),
+	FVector2D(0.03632313013f, 0.216376245f),
+	FVector2D(0.8285242319f, 0.9543109536f),
+	FVector2D(0.6305007339f, 0.9932333231f),
+	FVector2D(0.03632313013f, 0.216376245f),
+	FVector2D(0.8285242319f, 0.9543109536f),
+	FVector2D(0.03632313013f, 0.216376245f),
+	FVector2D(0.948182106f, 0.8630086184f),
+	FVector2D(0.9675967693f, 0.8031471968f),
+	FVector2D(0.03632313013f, 0.216376245f),
+	FVector2D(0.948182106f, 0.8630086184f),
+	FVector2D(0.03632313013f, 0.216376245f),
+	FVector2D(1.0f, 0.6263416409f),
+	FVector2D(0.9675967693f, 0.8031471968f),
+	FVector2D(0.03632313013f, 0.216376245f),
+	FVector2D(0.9872991443f, 0.3334831297f),
+	FVector2D(1.0f, 0.6263416409f),
+	FVector2D(0.03632313013f, 0.216376245f),
+	FVector2D(0.9275138378f, 0.1436391324f),
+	FVector2D(0.9872991443f, 0.3334831297f),
+	FVector2D(0.03632313013f, 0.216376245f),
+	FVector2D(0.8197077513f, 0.05096876994f),
+	FVector2D(0.9275138378f, 0.1436391324f),
+	FVector2D(0.7506074309f, 0.0313315168f),
+	FVector2D(0.03632313013f, 0.216376245f),
+	FVector2D(0.8197077513f, 0.05096876994f),
+	FVector2D(0.03632313013f, 0.216376245f),
+	FVector2D(0.5661634207f, 0.0f),
+	FVector2D(0.7506074309f, 0.0313315168f),
+	FVector2D(0.03632313013f, 0.216376245f),
+	FVector2D(0.3018191457f, 0.005146139301f),
+	FVector2D(0.5661634207f, 0.0f),
+	FVector2D(0.03632313013f, 0.216376245f),
+	FVector2D(0.06468880177f, 0.131273821f),
+	FVector2D(0.3018191457f, 0.005146139301f),
+	FVector2D(0.06468880177f, 0.131273821f),
+	FVector2D(0.1512531638f, 0.03234770149f),
+	FVector2D(0.3018191457f, 0.005146139301f),
+};
+
 static const FVector2D DK2_LeftEyeVisibleAreaPositions[] =
 {
 	FVector2D(0.1548848748f, 0.04138758406f),
diff --git a/Engine/Plugins/Runtime/Steam/SteamVR/Source/SteamVR/Private/SteamVRHMD.cpp b/Engine/Plugins/Runtime/Steam/SteamVR/Source/SteamVR/Private/SteamVRHMD.cpp
--- a/Engine/Plugins/Runtime/Steam/SteamVR/Source/SteamVR/Private/SteamVRHMD.cpp
+++ b/Engine/Plugins/Runtime/Steam/SteamVR/Source/SteamVR/Private/SteamVRHMD.cpp
@@ -1157,6 +1157,8 @@ void FSteamVRHMD::SetupOcclusionMeshes()
 		HiddenAreaMeshes[1].BuildMesh(RightEyePositions, VertexCount, FHMDViewMesh::MT_HiddenArea);
 
 		// If the hidden area mesh from the SteamVR runtime matches the mesh used to generate the Vive's visible area mesh, initialize it.
+		// The visible area mesh is a hand crafted inverse of the hidden area mesh we are getting from the steamvr runtime. Since the runtime data
+		// may change, we need to sanity check it matches our hand crafted mesh before using it.
 		if (HiddenAreaMeshCrc == ViveHiddenAreaMeshCrc)
 		{
 			VisibleAreaMeshes[0].BuildMesh(Vive_LeftEyeVisibleAreaPositions, VisibleAreaVertexCount, FHMDViewMesh::MT_VisibleArea);
diff --git a/Engine/Source/Runtime/Renderer/Private/DeferredShadingRenderer.cpp b/Engine/Source/Runtime/Renderer/Private/DeferredShadingRenderer.cpp
--- a/Engine/Source/Runtime/Renderer/Private/DeferredShadingRenderer.cpp
+++ b/Engine/Source/Runtime/Renderer/Private/DeferredShadingRenderer.cpp
@@ -949,19 +949,15 @@ void FDeferredShadingSceneRenderer::Render(FRHICommandListImmediate& RHICmdList)
 	GRenderTargetPool.AddPhaseEvent(TEXT("EarlyZPass"));
 
 	// Draw the scene pre-pass / early z pass, populating the scene depth buffer and HiZ
-	bool bDepthWasCleared = false;
+	bool bDepthWasCleared = RenderPrePassHMD(RHICmdList);;
 	const bool bNeedsPrePass = NeedsPrePass(this);
 	if (bNeedsPrePass)
 	{
 		RHICmdList.SetCurrentStat(GET_STATID(STAT_CLM_PrePass));
-		RenderPrePass(RHICmdList);
+		RenderPrePass(RHICmdList, bDepthWasCleared);
 		// at this point, the depth was cleared
 		bDepthWasCleared = true;
 	}
-	else if (HasHiddenAreaMask())
-	{
-		bDepthWasCleared = RenderPrePassHMD(RHICmdList);
-	}
 
 	RHICmdList.SetCurrentStat(GET_STATID(STAT_CLM_AfterPrePass));
 	ServiceLocalQueue();
@@ -1443,12 +1439,6 @@ bool FDeferredShadingSceneRenderer::RenderPrePassView(FRHICommandList& RHICmdLis
 
 	SetupPrePassView(RHICmdList, View.ViewRect);
 
-	if (HasHiddenAreaMask() && View.StereoPass != eSSP_FULL)
-	{
-		RenderHiddenAreaMaskView(RHICmdList, View);
-		bDirty = true;
-	}
-	
 	// Draw the static occluder primitives using a depth drawing policy.
 	{
 		// Draw opaque occluders which support a separate position-only
@@ -1530,11 +1520,6 @@ public:
 	{
 		FSceneRenderTargets::Get(CmdList).BeginRenderingPrePass(CmdList, false);
 		SetupPrePassView(CmdList, View.ViewRect);
-
-		if (HasHiddenAreaMask() && View.StereoPass != eSSP_FULL)
-		{
-			RenderHiddenAreaMaskView(CmdList, View);
-		}
 	}
 };
 
@@ -1595,7 +1580,7 @@ void FDeferredShadingSceneRenderer::RenderPrePassViewParallel(const FViewInfo& V
 }
 
 /** Renders the scene's prepass and occlusion queries */
-bool FDeferredShadingSceneRenderer::RenderPrePass(FRHICommandListImmediate& RHICmdList)
+bool FDeferredShadingSceneRenderer::RenderPrePass(FRHICommandListImmediate& RHICmdList, bool bDepthWasCleared)
 {
 	SCOPED_DRAW_EVENT(RHICmdList, PrePass);
 	SCOPE_CYCLE_COUNTER(STAT_DepthDrawTime);
@@ -1603,7 +1588,7 @@ bool FDeferredShadingSceneRenderer::RenderPrePass(FRHICommandListImmediate& RHIC
 	bool bDirty = false;
 	FSceneRenderTargets& SceneContext = FSceneRenderTargets::Get(RHICmdList);
 
-	SceneContext.BeginRenderingPrePass(RHICmdList, true);
+	SceneContext.BeginRenderingPrePass(RHICmdList, !bDepthWasCleared);
 
 	// Draw a depth pass to avoid overdraw in the other passes.
 	if(EarlyZPassMode != DDM_None)
diff --git a/Engine/Source/Runtime/Renderer/Private/DeferredShadingRenderer.h b/Engine/Source/Runtime/Renderer/Private/DeferredShadingRenderer.h
--- a/Engine/Source/Runtime/Renderer/Private/DeferredShadingRenderer.h
+++ b/Engine/Source/Runtime/Renderer/Private/DeferredShadingRenderer.h
@@ -161,7 +161,7 @@ private:
 	 * Renders the scene's prepass and occlusion queries.
 	 * @return true if anything was rendered
 	 */
-	bool RenderPrePass(FRHICommandListImmediate& RHICmdList);
+	bool RenderPrePass(FRHICommandListImmediate& RHICmdList, bool bDepthWasCleared);
 
 	/**
 	 * Renders the active HMD's hidden area mask as a depth prepass, if available.
diff --git a/Engine/Source/Runtime/Renderer/Private/PostProcess/PostProcessTonemap.cpp b/Engine/Source/Runtime/Renderer/Private/PostProcess/PostProcessTonemap.cpp
--- a/Engine/Source/Runtime/Renderer/Private/PostProcess/PostProcessTonemap.cpp
+++ b/Engine/Source/Runtime/Renderer/Private/PostProcess/PostProcessTonemap.cpp
@@ -1099,7 +1099,12 @@ void FRCPassPostProcessTonemap::Process(FRenderingCompositePassContext& Context)
 	// Set the view family's render target/viewport.
 	SetRenderTarget(Context.RHICmdList, DestRenderTarget.TargetableTexture, FTextureRHIParamRef(), ESimpleRenderTargetMode::EUninitializedColorAndDepth);
 
-	if( ViewFamily.RenderTarget->GetRenderTargetTexture() != DestRenderTarget.TargetableTexture )
+	if (Context.HasHmdMesh() && View.StereoPass == eSSP_LEFT_EYE)
+	{
+		// needed when using an hmd mesh instead of a full screen quad because we don't touch all of the pixels in the render target
+		Context.RHICmdList.Clear(true, FLinearColor::Black, false, 1.0f, false, 0, FIntRect());
+	}
+	else if( ViewFamily.RenderTarget->GetRenderTargetTexture() != DestRenderTarget.TargetableTexture )
 	{
 		// needed to not have PostProcessAA leaking in content (e.g. Matinee black borders), is optimized away if possible (RT size=view size, )
 		Context.RHICmdList.Clear(true, FLinearColor::Black, false, 1.0f, false, 0, View.ViewRect);
