winxp cis fix

diff --git a/Engine/Plugins/Runtime/OculusRift/Source/OculusRift/Private/OculusRiftLayers.cpp b/Engine/Plugins/Runtime/OculusRift/Source/OculusRift/Private/OculusRiftLayers.cpp
--- a/Engine/Plugins/Runtime/OculusRift/Source/OculusRift/Private/OculusRiftLayers.cpp
+++ b/Engine/Plugins/Runtime/OculusRift/Source/OculusRift/Private/OculusRiftLayers.cpp
@@ -4,6 +4,8 @@
 #include "OculusRiftHMD.h"
 #include "OculusRiftLayers.h"
 
+#if OCULUS_RIFT_SUPPORTED_PLATFORMS
+
 FRenderLayer::FRenderLayer(FHMDLayerDesc& InDesc) :
 	FHMDRenderLayer(InDesc)
 {
@@ -255,3 +257,4 @@ ovrResult FLayerManager::SubmitFrame(FViewExtension* RenderContext)
 	return res;
 }
 
+#endif //OCULUS_RIFT_SUPPORTED_PLATFORMS
