Oculus: Revert change to disallow plugin initialization when device is disconnected

diff --git a/Engine/Plugins/Runtime/OculusRift/Source/OculusRift/Private/OculusRiftHMD.cpp b/Engine/Plugins/Runtime/OculusRift/Source/OculusRift/Private/OculusRiftHMD.cpp
--- a/Engine/Plugins/Runtime/OculusRift/Source/OculusRift/Private/OculusRiftHMD.cpp
+++ b/Engine/Plugins/Runtime/OculusRift/Source/OculusRift/Private/OculusRiftHMD.cpp
@@ -56,7 +56,7 @@ bool FOculusRiftPlugin::Initialize()
 
 #if OCULUS_RIFT_SUPPORTED_PLATFORMS
 		// Only init module when running Game or Editor, and Oculus service is running
-		if (!IsRunningDedicatedServer() && ovr_Detect(0).IsOculusHMDConnected)
+		if (!IsRunningDedicatedServer() && ovr_Detect(0).IsOculusServiceRunning)
 		{
 			ovrInitParams initParams;
 			FMemory::Memset(initParams, 0);
