[INTEGRATE] Merging over fix for broken unity builds in Oculus plugin
- Merged CL 2938175 from Dev-VR

diff --git a/Engine/Plugins/Runtime/OculusRift/Source/OculusRift/Private/OculusRiftRenderD3D12.cpp b/Engine/Plugins/Runtime/OculusRift/Source/OculusRift/Private/OculusRiftRenderD3D12.cpp
--- a/Engine/Plugins/Runtime/OculusRift/Source/OculusRift/Private/OculusRiftRenderD3D12.cpp
+++ b/Engine/Plugins/Runtime/OculusRift/Source/OculusRift/Private/OculusRiftRenderD3D12.cpp
@@ -7,9 +7,44 @@
 
 #if defined(OVR_D3D)
 
+#include "D3D11RHIPrivate.h"
+#include "D3D11Util.h"
+
+#define GetD3D11CubeFace GetD3D12CubeFace
+#define VerifyD3D11Result VerifyD3D12Result
+#define GetD3D11TextureFromRHITexture GetD3D12TextureFromRHITexture
+#define FRingAllocation FRingAllocation_D3D12
+#define GetRenderTargetFormat GetRenderTargetFormat_D3D12
+#define ED3D11ShaderOffsetBuffer ED3D12ShaderOffsetBuffer
+#define FindShaderResourceDXGIFormat FindShaderResourceDXGIFormat_D3D12
+#define FindUnorderedAccessDXGIFormat FindUnorderedAccessDXGIFormat_D3D12
+#define FindDepthStencilDXGIFormat FindDepthStencilDXGIFormat_D3D12
+#define HasStencilBits HasStencilBits_D3D12
+#define FVector4VertexDeclaration FVector4VertexDeclaration_D3D12
+#define GLOBAL_CONSTANT_BUFFER_INDEX GLOBAL_CONSTANT_BUFFER_INDEX_D3D12
+#define MAX_CONSTANT_BUFFER_SLOTS MAX_CONSTANT_BUFFER_SLOTS_D3D12
+#define FD3DGPUProfiler FD3D12GPUProfiler
+#define FRangeAllocator FRangeAllocator_D3D12
+
 #include "D3D12RHIPrivate.h"
 #include "D3D12Util.h"
 
+#undef GetD3D11CubeFace
+#undef VerifyD3D11Result
+#undef GetD3D11TextureFromRHITexture
+#undef FRingAllocation
+#undef GetRenderTargetFormat
+#undef ED3D11ShaderOffsetBuffer
+#undef FindShaderResourceDXGIFormat
+#undef FindUnorderedAccessDXGIFormat
+#undef FindDepthStencilDXGIFormat
+#undef HasStencilBits
+#undef FVector4VertexDeclaration
+#undef GLOBAL_CONSTANT_BUFFER_INDEX
+#undef MAX_CONSTANT_BUFFER_SLOTS
+#undef FD3DGPUProfiler
+#undef FRangeAllocator
+
 #ifndef WINDOWS_PLATFORM_TYPES_GUARD
 #include "AllowWindowsPlatformTypes.h"
 #endif
