Fix for WinXP build breaks with Oculus 1.3

diff --git a/Engine/Plugins/Runtime/OculusRift/Source/OculusRift/Private/OculusRiftHMD.cpp b/Engine/Plugins/Runtime/OculusRift/Source/OculusRift/Private/OculusRiftHMD.cpp
--- a/Engine/Plugins/Runtime/OculusRift/Source/OculusRift/Private/OculusRiftHMD.cpp
+++ b/Engine/Plugins/Runtime/OculusRift/Source/OculusRift/Private/OculusRiftHMD.cpp
@@ -84,6 +84,7 @@ bool FOculusRiftPlugin::Initialize()
 	return bInitialized;
 }
 
+#if OCULUS_RIFT_SUPPORTED_PLATFORMS
 ovrResult FOculusRiftPlugin::CreateSession(ovrSession* session, ovrGraphicsLuid* luid)
 {
 	ovrResult result;
@@ -135,6 +136,7 @@ void FOculusRiftPlugin::DestroySession(ovrSession session)
 {
 	ovr_Destroy(session);
 }
+#endif //OCULUS_RIFT_SUPPORTED_PLATFORMS
 
 void FOculusRiftPlugin::ShutdownModule()
 {
@@ -227,12 +229,12 @@ TSharedPtr< class IHeadMountedDisplay, ESPMode::ThreadSafe > FOculusRiftPlugin::
 	return NULL;
 }
 
+#if OCULUS_RIFT_SUPPORTED_PLATFORMS
 bool FOculusRiftPlugin::PoseToOrientationAndPosition(const ovrPosef& Pose, FQuat& OutOrientation, FVector& OutPosition) const
 {
 	check(IsInGameThread());
 	bool bRetVal = false;
 
-#if OCULUS_RIFT_SUPPORTED_PLATFORMS
 	IHeadMountedDisplay* HMD = HeadMountedDisplay.Pin().Get();
 	if (HMD && HMD->GetHMDDeviceType() == EHMDDeviceType::DT_OculusRift)
 	{
@@ -245,7 +247,7 @@ bool FOculusRiftPlugin::PoseToOrientationAndPosition(const ovrPosef& Pose, FQuat
 			bRetVal = true;
 		}
 	}
-#endif //OCULUS_RIFT_SUPPORTED_PLATFORMS
+
 	return bRetVal;
 }
 
@@ -253,14 +255,13 @@ class FOvrSessionShared* FOculusRiftPlugin::GetSession()
 {
 	check(IsInGameThread());
 
-#if OCULUS_RIFT_SUPPORTED_PLATFORMS
 	IHeadMountedDisplay* HMD = HeadMountedDisplay.Pin().Get();
 	if (HMD && HMD->GetHMDDeviceType() == EHMDDeviceType::DT_OculusRift)
 	{
 		FOculusRiftHMD* OculusHMD = static_cast<FOculusRiftHMD*>(HMD);
 		return OculusHMD->Session.Get();
 	}
-#endif //OCULUS_RIFT_SUPPORTED_PLATFORMS
+
 	return nullptr;
 }
 
@@ -269,7 +270,6 @@ bool FOculusRiftPlugin::GetCurrentTrackingState(ovrTrackingState* TrackingState)
 	check(IsInGameThread());
 	bool bRetVal = false;
 
-#if OCULUS_RIFT_SUPPORTED_PLATFORMS
 	IHeadMountedDisplay* HMD = HeadMountedDisplay.Pin().Get();
 	if (TrackingState && HMD && HMD->GetHMDDeviceType() == EHMDDeviceType::DT_OculusRift)
 	{
@@ -281,9 +281,10 @@ bool FOculusRiftPlugin::GetCurrentTrackingState(ovrTrackingState* TrackingState)
 			bRetVal = true;
 		}
 	}
-#endif //OCULUS_RIFT_SUPPORTED_PLATFORMS
+
 	return bRetVal;
 }
+#endif //OCULUS_RIFT_SUPPORTED_PLATFORMS
 
 IMPLEMENT_MODULE( FOculusRiftPlugin, OculusRift )
 
diff --git a/Engine/Plugins/Runtime/OculusRift/Source/OculusRift/Private/OculusRiftHMD.h b/Engine/Plugins/Runtime/OculusRift/Source/OculusRift/Private/OculusRiftHMD.h
--- a/Engine/Plugins/Runtime/OculusRift/Source/OculusRift/Private/OculusRiftHMD.h
+++ b/Engine/Plugins/Runtime/OculusRift/Source/OculusRift/Private/OculusRiftHMD.h
@@ -85,7 +85,7 @@ FORCEINLINE OVR::Recti ToOVRRecti(const FIntRect& rect)
 {
 	return OVR::Recti(rect.Min.X, rect.Min.Y, rect.Size().X, rect.Size().Y);
 }
-
+#endif //OCULUS_RIFT_SUPPORTED_PLATFORMS
 
 //-------------------------------------------------------------------------------------------------
 // FOculusRiftPlugin
@@ -110,9 +110,10 @@ public:
 		return FModuleManager::LoadModuleChecked< FOculusRiftPlugin >( "OculusRift" );
 	}
 
+#if OCULUS_RIFT_SUPPORTED_PLATFORMS
 	ovrResult CreateSession(ovrSession* session, ovrGraphicsLuid* luid);
 	static void DestroySession(ovrSession session);
-
+#endif //OCULUS_RIFT_SUPPORTED_PLATFORMS
 public:
 	/** IModuleInterface */
 	virtual void ShutdownModule() override;
@@ -127,9 +128,11 @@ public:
 	virtual TSharedPtr< class IHeadMountedDisplay, ESPMode::ThreadSafe > CreateHeadMountedDisplay() override;
 
 	/** IOculusRiftPlugin */
+#if OCULUS_RIFT_SUPPORTED_PLATFORMS
 	virtual bool PoseToOrientationAndPosition(const ovrPosef& Pose, FQuat& OutOrientation, FVector& OutPosition) const override;
 	virtual class FOvrSessionShared* GetSession() override;
 	virtual bool GetCurrentTrackingState(ovrTrackingState* TrackingState) override;
+#endif //OCULUS_RIFT_SUPPORTED_PLATFORMS
 
 protected:
 	bool bInitialized;
@@ -142,6 +145,7 @@ protected:
 // FSettings
 //-------------------------------------------------------------------------------------------------
 
+#if OCULUS_RIFT_SUPPORTED_PLATFORMS
 class FSettings : public FHMDSettings
 {
 public:
diff --git a/Engine/Plugins/Runtime/OculusRift/Source/OculusRift/Public/IOculusRiftPlugin.h b/Engine/Plugins/Runtime/OculusRift/Source/OculusRift/Public/IOculusRiftPlugin.h
--- a/Engine/Plugins/Runtime/OculusRift/Source/OculusRift/Public/IOculusRiftPlugin.h
+++ b/Engine/Plugins/Runtime/OculusRift/Source/OculusRift/Public/IOculusRiftPlugin.h
@@ -38,6 +38,7 @@ public:
 		return FModuleManager::Get().IsModuleLoaded( "OculusRift" );
 	}
 
+#if OCULUS_RIFT_SUPPORTED_PLATFORMS
 	/**
 	 * Takes an orientation and position in Oculus Rift coordinate space and converts it to 
 	 * Unreal's coordinate system, also applying world scale and base position/orientation offsets
@@ -53,5 +54,6 @@ public:
 	 * Returns current ovrTrackingState
 	 */
 	virtual bool GetCurrentTrackingState(struct ovrTrackingState_* TrackingState) = 0;
+#endif //OCULUS_RIFT_SUPPORTED_PLATFORMS
 };
 
