Compiler fix for Win32

diff --git a/Engine/Plugins/Runtime/OculusRift/Source/OculusRift/Private/SceneCubemapCapturer.h b/Engine/Plugins/Runtime/OculusRift/Source/OculusRift/Private/SceneCubemapCapturer.h
--- a/Engine/Plugins/Runtime/OculusRift/Source/OculusRift/Private/SceneCubemapCapturer.h
+++ b/Engine/Plugins/Runtime/OculusRift/Source/OculusRift/Private/SceneCubemapCapturer.h
@@ -37,7 +37,7 @@ public:
 	void SetOffset(FVector InOffset) { CaptureOffset = InOffset; }
 
 	// overrides player's 0 orientation for the capture. 
-	void SetInitialOrientation(FQuat InOrientation) { OverriddenOrientation = InOrientation; }
+	void SetInitialOrientation(const FQuat& InOrientation) { OverriddenOrientation = InOrientation; }
 
 	// overrides player's 0 location for the capture.
 	void SetInitialLocation(FVector InLocation) { OverriddenLocation = InLocation; }
