Removing left eye overscan changes

diff --git a/Engine/Plugins/Runtime/OculusRift/Source/OculusRift/Private/OculusRiftHMD.cpp b/Engine/Plugins/Runtime/OculusRift/Source/OculusRift/Private/OculusRiftHMD.cpp
--- a/Engine/Plugins/Runtime/OculusRift/Source/OculusRift/Private/OculusRiftHMD.cpp
+++ b/Engine/Plugins/Runtime/OculusRift/Source/OculusRift/Private/OculusRiftHMD.cpp
@@ -26,9 +26,6 @@
 
 static TAutoConsoleVariable<int32> CStartInVRVar(TEXT("vr.bStartInVR"), 0, TEXT("Start in VR flag"));
 static TAutoConsoleVariable<int32> CGracefulExitVar(TEXT("vr.bExitGracefully"), 0, TEXT("Exit gracefully when forced by Universal Menu."));
-static TAutoConsoleVariable<float> CRightEyeResolutionScaling( TEXT( "rift.RightEyeResolutionScaling" ), 1.0f, TEXT( "Scales the resolution of the right eye only" ) );
-static TAutoConsoleVariable<int32> CLeftEyeOverscan( TEXT( "rift.LeftEyeOverscan" ), 0, TEXT( "Whether to render the left eye at a larger widescreen resolution, then crop a column from the center for the HMD.  This is a presentation mode feature." ) );
-static TAutoConsoleVariable<float> CLeftEyeOverscanAspect( TEXT( "rift.LeftEyeOverscanAspect" ), 16.0f / 9.0f, TEXT( "Aspect ratio when rift.LeftEyeOverscan is enabled" ) );
 
 //---------------------------------------------------
 // Oculus Rift Plugin Implementation
@@ -2142,6 +2139,10 @@ void FOculusRiftHMD::UpdateHmdRenderInfo()
 		// Calc FOV, symmetrical, for each eye. 
 		CurrentSettings->EyeFov[0] = HmdDesc.DefaultEyeFov[0];
 		CurrentSettings->EyeFov[1] = HmdDesc.DefaultEyeFov[1];
+
+		// Calc FOV in radians
+		CurrentSettings->VFOVInRadians = GetVerticalFovRadians(CurrentSettings->EyeFov[0], CurrentSettings->EyeFov[1]);
+		CurrentSettings->HFOVInRadians = GetHorizontalFovRadians(CurrentSettings->EyeFov[0], CurrentSettings->EyeFov[1]);
 	}
 
 	const ovrSizei recommenedTex0Size = ovr_GetFovTextureSize(OvrSession, ovrEye_Left, CurrentSettings->EyeFov[0], 1.0f);
@@ -2154,27 +2155,6 @@ void FOculusRiftHMD::UpdateHmdRenderInfo()
 	CurrentSettings->IdealScreenPercentage = FMath::Max(float(idealRenderTargetSize.w) / float(HmdDesc.Resolution.w) * 100.f,
 														float(idealRenderTargetSize.h) / float(HmdDesc.Resolution.h) * 100.f);
 
-	if (!CurrentSettings->Flags.bOverrideFOV)
-	{
-		if( CLeftEyeOverscan->GetInt() != 0 )
-		{
-			// Compute a new FOV for the left eye that will "overscan" the HMD's viewport to allow for a
-			// much wider mirror image.  The user of the HMD will see the horizontally cropped section of
-			// view that would normally be rendered without the overscan, but viewers of the mirror display
-			// will see the entire view.
-			const ovrSizei OriginalEyeRenderTargetSize = ovr_GetFovTextureSize( OvrSession, ovrEye_Left, CurrentSettings->EyeFov[ 0 ], 1.0f );
-			const float OriginalEyeAspect = (float)OriginalEyeRenderTargetSize.w / (float)OriginalEyeRenderTargetSize.h;
-			const float NewEyeAspect = CLeftEyeOverscanAspect->GetFloat();
-
-			CurrentSettings->EyeFov[ 0 ].LeftTan *= NewEyeAspect / OriginalEyeAspect;
-			CurrentSettings->EyeFov[ 0 ].RightTan *= NewEyeAspect / OriginalEyeAspect;
-		}
-
-		// Calc FOV in radians
-		CurrentSettings->VFOVInRadians = GetVerticalFovRadians(CurrentSettings->EyeFov[0], CurrentSettings->EyeFov[1]);
-		CurrentSettings->HFOVInRadians = GetHorizontalFovRadians(CurrentSettings->EyeFov[0], CurrentSettings->EyeFov[1]);
-	}
-
 	// Override eye distance by the value from HMDInfo (stored in Profile).
 	if (!CurrentSettings->Flags.bOverrideIPD)
 	{
@@ -2264,13 +2244,7 @@ void FOculusRiftHMD::UpdateStereoRenderingParams()
 
 		for(int eyeIndex = 0; eyeIndex < ovrEye_Count; eyeIndex++)
 		{
-			float eyePixelDensity = pixelDensity;
-			if( eyeIndex == ovrEye_Right )
-			{
-				eyePixelDensity *= CRightEyeResolutionScaling->GetFloat();
-			}
-
-			vpSize[eyeIndex] = ovr_GetFovTextureSize(OvrSession, (ovrEyeType) eyeIndex, CurrentSettings->EyeFov[eyeIndex], eyePixelDensity);
+			vpSize[eyeIndex] = ovr_GetFovTextureSize(OvrSession, (ovrEyeType)eyeIndex, CurrentSettings->EyeFov[eyeIndex], pixelDensity);
 			vpSize[eyeIndex].w = FMath::Max(vpSize[eyeIndex].w, 1);
 			vpSize[eyeIndex].h = FMath::Max(vpSize[eyeIndex].h, 1);
 
diff --git a/Engine/Source/Runtime/Engine/Private/HierarchicalInstancedStaticMesh.cpp b/Engine/Source/Runtime/Engine/Private/HierarchicalInstancedStaticMesh.cpp
--- a/Engine/Source/Runtime/Engine/Private/HierarchicalInstancedStaticMesh.cpp
+++ b/Engine/Source/Runtime/Engine/Private/HierarchicalInstancedStaticMesh.cpp
@@ -1332,9 +1332,6 @@ void FHierarchicalStaticMeshSceneProxy::GetDynamicMeshElements(const TArray<cons
 		{
 			const FSceneView* View = Views[ViewIndex];
 
-			// With stereo rendering views may have asymmetric FOVs, so always use the left eye's view for LOD calculation
-			const FSceneView* ViewForLOD = ( View->Family != nullptr && View->StereoPass == eSSP_RIGHT_EYE ) ? View->Family->Views[ 0 ] : View;
-
 			FFoliageElementParams ElementParams;
 			ElementParams.bSelectionRenderEnabled = GIsEditor && ViewFamily.EngineShowFlags.Selection;
 			ElementParams.NumSelectionGroups = (ElementParams.bSelectionRenderEnabled && bHasSelectedInstances) ? 2 : 1;
@@ -1403,7 +1400,7 @@ void FHierarchicalStaticMeshSceneProxy::GetDynamicMeshElements(const TArray<cons
 					}
 					else
 					{
-						const FMatrix LocalViewProjForCulling = GetLocalToWorld() * ViewForLOD->ViewMatrices.GetViewProjectionMatrix();
+						const FMatrix LocalViewProjForCulling = GetLocalToWorld() * View->ViewMatrices.GetViewProjectionMatrix();
 						GetViewFrustumBounds(InstanceParams.ViewFrustumLocal, LocalViewProjForCulling, false);
 					}
 
diff --git a/Engine/Source/Runtime/Engine/Private/InstancedStaticMesh.cpp b/Engine/Source/Runtime/Engine/Private/InstancedStaticMesh.cpp
--- a/Engine/Source/Runtime/Engine/Private/InstancedStaticMesh.cpp
+++ b/Engine/Source/Runtime/Engine/Private/InstancedStaticMesh.cpp
@@ -2013,9 +2013,6 @@ void FInstancedStaticMeshVertexFactoryShaderParameters::SetMesh( FRHICommandList
 		InstancingWorldViewOriginOne.W = 1.0f;
 		if (InstancingUserData && BatchElement.InstancedLODRange)
 		{
-			// With stereo rendering views may have asymmetric FOVs, so always use the left eye's view for LOD calculation
-			const FSceneView& ViewForLOD = ( View.Family != nullptr && View.StereoPass == eSSP_RIGHT_EYE ) ? *View.Family->Views[ 0 ] : View;
-
 			int32 FirstLOD = InstancingUserData->MinLOD;
 
 			int32 DebugMin = FMath::Min(CVarMinLOD.GetValueOnRenderThread(), InstancingUserData->MeshRenderData->LODResources.Num() - 1);
diff --git a/Engine/Source/Runtime/Engine/Private/SceneManagement.cpp b/Engine/Source/Runtime/Engine/Private/SceneManagement.cpp
--- a/Engine/Source/Runtime/Engine/Private/SceneManagement.cpp
+++ b/Engine/Source/Runtime/Engine/Private/SceneManagement.cpp
@@ -375,9 +375,7 @@ int8 ComputeTemporalStaticMeshLOD( const FStaticMeshRenderData* RenderData, cons
 {
 	const int32 NumLODs = MAX_STATIC_MESH_LODS;
 
-	// With stereo rendering views may have asymmetric FOVs, so always use the left eye's view for LOD calculation
-	const FSceneView& ViewForLOD = ( View.Family != nullptr && View.StereoPass == eSSP_RIGHT_EYE ) ? *View.Family->Views[ 0 ] : View;
-	const float ScreenRadiusSquared = ComputeTemporalLODBoundsScreenRadiusSquared(Origin, SphereRadius, ViewForLOD, SampleIndex) * FactorScale * FactorScale;
+	const float ScreenRadiusSquared = ComputeTemporalLODBoundsScreenRadiusSquared(Origin, SphereRadius, View, SampleIndex) * FactorScale * FactorScale;
 
 	// Walk backwards and return the first matching LOD
 	for(int32 LODIndex = NumLODs - 1 ; LODIndex >= 0 ; --LODIndex)
@@ -397,9 +395,7 @@ int8 ComputeStaticMeshLOD( const FStaticMeshRenderData* RenderData, const FVecto
 	{
 		const int32 NumLODs = MAX_STATIC_MESH_LODS;
 
-		// With stereo rendering views may have asymmetric FOVs, so always use the left eye's view for LOD calculation
-		const FSceneView& ViewForLOD = ( View.Family != nullptr && View.StereoPass == eSSP_RIGHT_EYE ) ? *View.Family->Views[ 0 ] : View;
-		const float ScreenRadiusSquared = ComputeBoundsScreenRadiusSquared(Origin, SphereRadius, ViewForLOD) * FactorScale * FactorScale;
+		const float ScreenRadiusSquared = ComputeBoundsScreenRadiusSquared(Origin, SphereRadius, View) * FactorScale * FactorScale;
 
 		// Walk backwards and return the first matching LOD
 		for (int32 LODIndex = NumLODs - 1; LODIndex >= 0; --LODIndex)
@@ -434,15 +430,13 @@ FLODMask ComputeLODForMeshes( const TIndirectArray<class FStaticMesh>& StaticMes
 	{
 		int32 NumMeshes = StaticMeshes.Num();
 
-		// With stereo rendering views may have asymmetric FOVs, so always use the left eye's view for LOD calculation
-		const FSceneView& ViewForLOD = ( View.Family != nullptr && View.StereoPass == eSSP_RIGHT_EYE ) ? *View.Family->Views[ 0 ] : View;
 		if (NumMeshes && StaticMeshes[0].bDitheredLODTransition)
 		{
 			for (int32 SampleIndex = 0; SampleIndex < 2; SampleIndex++)
 			{
 				int32 MinLODFound = INT_MAX;
 				bool bFoundLOD = false;
-				const float ScreenRadiusSquared = ComputeTemporalLODBoundsScreenRadiusSquared(Origin, SphereRadius, ViewForLOD, SampleIndex);
+				const float ScreenRadiusSquared = ComputeTemporalLODBoundsScreenRadiusSquared(Origin, SphereRadius, View, SampleIndex);
 
 				for(int32 MeshIndex = NumMeshes-1 ; MeshIndex >= 0 ; --MeshIndex)
 				{
@@ -470,7 +464,7 @@ FLODMask ComputeLODForMeshes( const TIndirectArray<class FStaticMesh>& StaticMes
 		{
 			int32 MinLODFound = INT_MAX;
 			bool bFoundLOD = false;
-			const float ScreenRadiusSquared = ComputeBoundsScreenRadiusSquared(Origin, SphereRadius, ViewForLOD);
+			const float ScreenRadiusSquared = ComputeBoundsScreenRadiusSquared(Origin, SphereRadius, View);
 
 			for(int32 MeshIndex = NumMeshes-1 ; MeshIndex >= 0 ; --MeshIndex)
 			{
diff --git a/Engine/Source/Runtime/Engine/Private/SceneView.cpp b/Engine/Source/Runtime/Engine/Private/SceneView.cpp
--- a/Engine/Source/Runtime/Engine/Private/SceneView.cpp
+++ b/Engine/Source/Runtime/Engine/Private/SceneView.cpp
@@ -797,16 +797,13 @@ static TAutoConsoleVariable<int32> CVarCompensateForFOV(
 
 float FSceneView::GetLODDistanceFactor() const
 {
-	// With stereo rendering views may have asymmetric FOVs, so always use the left eye's view for LOD calculation
-	const FSceneView* ViewForLOD = ( this->Family != nullptr && this->StereoPass == eSSP_RIGHT_EYE ) ? Family->Views[0] : this;
-
 	bool bCompensateForFOV = bUseFieldOfViewForLOD && CVarCompensateForFOV.GetValueOnAnyThread() != 0;
-	float ScreenScaleX = bCompensateForFOV ? ViewForLOD->ViewMatrices.GetProjectionMatrix().M[0][0] : 1.0f;
-	float ScreenScaleY = bCompensateForFOV ? ViewForLOD->ViewMatrices.GetProjectionMatrix().M[1][1] : (static_cast<float>(ViewForLOD->ViewRect.Width()) / ViewForLOD->ViewRect.Height());
+	float ScreenScaleX = bCompensateForFOV ? ViewMatrices.GetProjectionMatrix().M[0][0] : 1.0f;
+	float ScreenScaleY = bCompensateForFOV ? ViewMatrices.GetProjectionMatrix().M[1][1] : (static_cast<float>(ViewRect.Width()) / ViewRect.Height());
 
-	const float ScreenMultiple = FMath::Max(ViewForLOD->ViewRect.Width() / 2.0f * ScreenScaleX,
-		ViewForLOD->ViewRect.Height() / 2.0f * ScreenScaleY);
-	float Fac = PI * ScreenMultiple * ScreenMultiple / ViewForLOD->ViewRect.Area();
+	const float ScreenMultiple = FMath::Max(ViewRect.Width() / 2.0f * ScreenScaleX,
+		ViewRect.Height() / 2.0f * ScreenScaleY);
+	float Fac = PI * ScreenMultiple * ScreenMultiple / ViewRect.Area();
 	return Fac;
 }
 
diff --git a/Engine/Source/Runtime/Engine/Private/SkeletalRender.cpp b/Engine/Source/Runtime/Engine/Private/SkeletalRender.cpp
--- a/Engine/Source/Runtime/Engine/Private/SkeletalRender.cpp
+++ b/Engine/Source/Runtime/Engine/Private/SkeletalRender.cpp
@@ -70,9 +70,7 @@ void FSkeletalMeshObject::UpdateMinDesiredLODLevel(const FSceneView* View, const
 	static const auto* SkeletalMeshLODRadiusScale = IConsoleManager::Get().FindTConsoleVariableDataFloat(TEXT("r.SkeletalMeshLODRadiusScale"));
 	float LODScale = FMath::Clamp(SkeletalMeshLODRadiusScale->GetValueOnRenderThread(), 0.25f, 1.0f);
 
-	// With stereo rendering views may have asymmetric FOVs, so always use the left eye's view for LOD calculation
-	const FSceneView* ViewForLOD = ( View->Family != nullptr && View->StereoPass == eSSP_RIGHT_EYE ) ? View->Family->Views[ 0 ] : View;
-	const float ScreenRadiusSquared = ComputeBoundsScreenRadiusSquared(Bounds.Origin, Bounds.SphereRadius, *ViewForLOD) * LODScale * LODScale;
+	const float ScreenRadiusSquared = ComputeBoundsScreenRadiusSquared(Bounds.Origin, Bounds.SphereRadius, *View) * LODScale * LODScale;
 
 	check( SkeletalMeshLODInfo.Num() == SkeletalMeshResource->LODModels.Num() );
 
